Restructure files in GizmoEdMode plugin

#rnx

diff --git a/Engine/Plugins/Experimental/GizmoEdMode/Source/GizmoEdMode.Build.cs b/Engine/Plugins/Experimental/GizmoEdMode/Source/GizmoEdMode.Build.cs
deleted file mode 100644
diff --git a/Engine/Plugins/Experimental/GizmoEdMode/Source/GizmoEdMode/GizmoEdMode.Build.cs b/Engine/Plugins/Experimental/GizmoEdMode/Source/GizmoEdMode/GizmoEdMode.Build.cs
new file mode 100644
--- /dev/null
+++ b/Engine/Plugins/Experimental/GizmoEdMode/Source/GizmoEdMode/GizmoEdMode.Build.cs
@@ -0,0 +1,57 @@
+// Copyright Epic Games, Inc. All Rights Reserved.
+
+using UnrealBuildTool;
+
+public class GizmoEdMode : ModuleRules
+{
+	public GizmoEdMode(ReadOnlyTargetRules Target) : base(Target)
+	{
+		PCHUsage = ModuleRules.PCHUsageMode.UseExplicitOrSharedPCHs;
+
+		PublicIncludePaths.AddRange(
+			new string[] {
+				// ... add public include paths required here ...
+			}
+			);
+
+
+		PrivateIncludePaths.AddRange(
+			new string[] {
+				// ... add other private include paths required here ...
+			}
+			);
+
+
+		PublicDependencyModuleNames.AddRange(
+			new string[]
+			{
+				// ... add other public dependencies that you statically link with here ...
+				"Core"
+			}
+			);
+
+
+		PrivateDependencyModuleNames.AddRange(
+			new string[]
+			{
+				"Engine",
+				"InputCore",
+				"EditorFramework",
+				"UnrealEd",
+				"SlateCore",
+				"CoreUObject",
+				"EditorInteractiveToolsFramework",
+				"InteractiveToolsFramework"
+				// ... add private dependencies that you statically link with here ...
+			}
+			);
+
+
+		DynamicallyLoadedModuleNames.AddRange(
+			new string[]
+			{
+				// ... add any modules that your module loads dynamically here ...
+			}
+			);
+	}
+}
diff --git a/Engine/Plugins/Experimental/GizmoEdMode/Source/GizmoEdMode/Private/DefaultAssetEditorGizmoFactory.cpp b/Engine/Plugins/Experimental/GizmoEdMode/Source/GizmoEdMode/Private/DefaultAssetEditorGizmoFactory.cpp
new file mode 100644
--- /dev/null
+++ b/Engine/Plugins/Experimental/GizmoEdMode/Source/GizmoEdMode/Private/DefaultAssetEditorGizmoFactory.cpp
@@ -0,0 +1,64 @@
+// Copyright Epic Games, Inc. All Rights Reserved.
+
+#include "DefaultAssetEditorGizmoFactory.h"
+#include "EditorModeManager.h"
+#include "Engine/Selection.h"
+#include "BaseGizmos/TransformGizmo.h"
+#include "BaseGizmos/TransformProxy.h"
+#include "UnrealWidget.h"
+
+bool UDefaultAssetEditorGizmoFactory::CanBuildGizmoForSelection(FEditorModeTools* ModeTools) const
+{
+	return true;
+}
+
+UTransformGizmo* UDefaultAssetEditorGizmoFactory::BuildGizmoForSelection(FEditorModeTools* ModeTools, UInteractiveGizmoManager* GizmoManager) const
+{
+	ETransformGizmoSubElements Elements  = ETransformGizmoSubElements::None;
+	bool bUseContextCoordinateSystem = true;
+	FWidget::EWidgetMode WidgetMode = ModeTools->GetWidgetMode();
+	switch ( WidgetMode )
+	{
+	case FWidget::EWidgetMode::WM_Translate:
+		Elements = ETransformGizmoSubElements::TranslateAllAxes | ETransformGizmoSubElements::TranslateAllPlanes;
+		break;
+	case FWidget::EWidgetMode::WM_Rotate:
+		Elements = ETransformGizmoSubElements::RotateAllAxes;
+		break;
+	case FWidget::EWidgetMode::WM_Scale:
+		Elements = ETransformGizmoSubElements::ScaleAllAxes | ETransformGizmoSubElements::ScaleAllPlanes;
+		bUseContextCoordinateSystem = false;
+		break;
+	case FWidget::EWidgetMode::WM_2D:
+		Elements = ETransformGizmoSubElements::RotateAxisY | ETransformGizmoSubElements::TranslatePlaneXZ;
+		break;
+	default:
+		Elements = ETransformGizmoSubElements::FullTranslateRotateScale;
+		break;
+	}
+	UTransformGizmo* TransformGizmo = GizmoManager->CreateCustomTransformGizmo(Elements);
+	TransformGizmo->bUseContextCoordinateSystem = bUseContextCoordinateSystem;
+
+	TArray<AActor*> SelectedActors;
+	ModeTools->GetSelectedActors()->GetSelectedObjects<AActor>(SelectedActors);
+
+	UTransformProxy* TransformProxy = NewObject<UTransformProxy>();
+	for (auto Actor : SelectedActors)
+	{
+		USceneComponent* SceneComponent = Actor->GetRootComponent();
+		TransformProxy->AddComponent(SceneComponent);
+	}
+	TransformGizmo->SetActiveTarget( TransformProxy );
+	TransformGizmo->SetVisibility(SelectedActors.Num() > 0);
+	return TransformGizmo;
+}
+
+void UDefaultAssetEditorGizmoFactory::ConfigureGridSnapping(bool bGridEnabled, bool bRotGridEnabled, UTransformGizmo* Gizmo) const
+{
+	if ( Gizmo )
+	{
+		Gizmo->bSnapToWorldGrid = bGridEnabled;
+		Gizmo->bSnapToWorldRotGrid = bRotGridEnabled;
+	}
+
+}
diff --git a/Engine/Plugins/Experimental/GizmoEdMode/Source/GizmoEdMode/Private/GizmoEdMode.cpp b/Engine/Plugins/Experimental/GizmoEdMode/Source/GizmoEdMode/Private/GizmoEdMode.cpp
new file mode 100644
--- /dev/null
+++ b/Engine/Plugins/Experimental/GizmoEdMode/Source/GizmoEdMode/Private/GizmoEdMode.cpp
@@ -0,0 +1,102 @@
+// Copyright Epic Games, Inc. All Rights Reserved.
+
+#include "GizmoEdMode.h"
+#include "DefaultAssetEditorGizmoFactory.h"
+#include "Components/SceneComponent.h"
+#include "EditorModeManager.h"
+#include "EdModeInteractiveToolsContext.h"
+#include "InteractiveGizmo.h"
+#include "BaseGizmos/TransformGizmo.h"
+#include "UnrealWidget.h"
+#include "Utils.h"
+
+#define LOCTEXT_NAMESPACE "FGizmoEdMode"
+
+UGizmoEdMode::UGizmoEdMode()
+	:Super()
+{
+	Info = FEditorModeInfo(
+		FName(TEXT("GizmoMode")),
+		LOCTEXT("ModeName", "Gizmo"),
+		FSlateIcon(),
+		false,
+		600
+	);
+	SettingsClass = UGizmoEdModeSettings::StaticClass();
+	ToolsContextClass = UEdModeInteractiveToolsContext::StaticClass();
+
+	AddFactory(NewObject<UDefaultAssetEditorGizmoFactory>());
+}
+
+void UGizmoEdMode::AddFactory(TScriptInterface<IAssetEditorGizmoFactory> GizmoFactory)
+{
+	GizmoFactories.Add(MoveTemp(GizmoFactory));
+	GizmoFactories.StableSort(
+	    [](const TScriptInterface<IAssetEditorGizmoFactory>& a, const TScriptInterface<IAssetEditorGizmoFactory>& b) {
+		    return a->GetPriority() > b->GetPriority();
+	    });
+}
+
+void UGizmoEdMode::ActorSelectionChangeNotify()
+{
+	RecreateGizmo();
+}
+
+void UGizmoEdMode::RecreateGizmo()
+{
+	DestroyGizmo();
+	for ( auto& Factory : GizmoFactories )
+	{
+		if (Factory->CanBuildGizmoForSelection(GetModeManager()))
+		{
+			TransformGizmo = Factory->BuildGizmoForSelection(GetModeManager(), ToolsContext->GizmoManager);
+			LastFactory = &(*Factory);
+			return;
+		}
+	}
+}
+
+void UGizmoEdMode::DestroyGizmo()
+{
+	LastFactory = nullptr;
+	if ( TransformGizmo )
+	{
+		ToolsContext->GizmoManager->DestroyGizmo(TransformGizmo);
+		TransformGizmo = nullptr;
+	}
+}
+
+void UGizmoEdMode::Enter()
+{
+	Super::Enter();
+	RecreateGizmo();
+	WidgetModeChangedHandle =
+	    GetModeManager()->OnWidgetModeChanged().AddLambda([this](FWidget::EWidgetMode) { RecreateGizmo(); });
+	GetModeManager()->SetShowWidget(false);
+}
+
+void UGizmoEdMode::Exit()
+{
+	DestroyGizmo();
+	GetModeManager()->OnWidgetModeChanged().Remove(WidgetModeChangedHandle);
+	WidgetModeChangedHandle.Reset();
+	GetModeManager()->SetShowWidget(true);
+	Super::Exit();
+}
+
+bool UGizmoEdMode::InputKey(FEditorViewportClient* ViewportClient, FViewport* Viewport, FKey Key, EInputEvent Event)
+{
+	return ToolsContext->InputKey(ViewportClient, Viewport, Key, Event);
+}
+
+void UGizmoEdMode::Tick(FEditorViewportClient* ViewportClient, float DeltaTime)
+{
+	Super::Tick(ViewportClient, DeltaTime);
+	if ( LastFactory )
+	{
+		LastFactory->ConfigureGridSnapping(GetDefault<ULevelEditorViewportSettings>()->GridEnabled,
+		                                   GetDefault<ULevelEditorViewportSettings>()->RotGridEnabled, TransformGizmo);
+	}
+}
+
+#undef LOCTEXT_NAMESPACE
diff --git a/Engine/Plugins/Experimental/GizmoEdMode/Source/GizmoEdMode/Private/GizmoEdModeModule.cpp b/Engine/Plugins/Experimental/GizmoEdMode/Source/GizmoEdMode/Private/GizmoEdModeModule.cpp
new file mode 100644
--- /dev/null
+++ b/Engine/Plugins/Experimental/GizmoEdMode/Source/GizmoEdMode/Private/GizmoEdModeModule.cpp
@@ -0,0 +1,25 @@
+// Copyright Epic Games, Inc. All Rights Reserved.
+
+#include "GizmoEdModeModule.h"
+#include "GizmoEdMode.h"
+#include "EditorModeRegistry.h"
+#include "Misc/CoreDelegates.h"
+
+#define LOCTEXT_NAMESPACE "FGizmoEdModeModule"
+
+void FGizmoEdModeModule::OnPostEngineInit()
+{
+}
+
+void FGizmoEdModeModule::StartupModule()
+{
+	FCoreDelegates::OnPostEngineInit.AddRaw(this, &FGizmoEdModeModule::OnPostEngineInit);
+}
+
+void FGizmoEdModeModule::ShutdownModule()
+{
+}
+
+#undef LOCTEXT_NAMESPACE
+
+IMPLEMENT_MODULE(FGizmoEdModeModule, GizmoEdMode);
diff --git a/Engine/Plugins/Experimental/GizmoEdMode/Source/GizmoEdMode/Public/AssetEditorGizmoFactory.h b/Engine/Plugins/Experimental/GizmoEdMode/Source/GizmoEdMode/Public/AssetEditorGizmoFactory.h
new file mode 100644
--- /dev/null
+++ b/Engine/Plugins/Experimental/GizmoEdMode/Source/GizmoEdMode/Public/AssetEditorGizmoFactory.h
@@ -0,0 +1,37 @@
+// Copyright Epic Games, Inc. All Rights Reserved.
+
+#pragma once
+
+#include "Tools/UEdMode.h"
+#include "UObject/Object.h"
+#include "UObject/Interface.h"
+#include "UnrealWidget.h"
+
+#include "AssetEditorGizmoFactory.generated.h"
+
+class UTransformGizmo;
+class UInteractiveGizmoManager;
+
+UENUM()
+enum class EAssetEditorGizmoFactoryPriority
+{
+	Default,
+	Normal,
+	High
+};
+
+UINTERFACE()
+class GIZMOEDMODE_API UAssetEditorGizmoFactory : public UInterface
+{
+	GENERATED_BODY()
+};
+
+class GIZMOEDMODE_API IAssetEditorGizmoFactory
+{
+	GENERATED_BODY()
+public:
+	virtual bool CanBuildGizmoForSelection(FEditorModeTools* ModeTools) const = 0;
+	virtual UTransformGizmo* BuildGizmoForSelection(FEditorModeTools* ModeTools, UInteractiveGizmoManager* GizmoManager) const = 0;
+	virtual EAssetEditorGizmoFactoryPriority GetPriority() const { return EAssetEditorGizmoFactoryPriority::Normal; }
+	virtual void ConfigureGridSnapping(bool bGridEnabled, bool bRotGridEnabled, UTransformGizmo* Gizmo) const = 0;
+};
diff --git a/Engine/Plugins/Experimental/GizmoEdMode/Source/GizmoEdMode/Public/DefaultAssetEditorGizmoFactory.h b/Engine/Plugins/Experimental/GizmoEdMode/Source/GizmoEdMode/Public/DefaultAssetEditorGizmoFactory.h
new file mode 100644
--- /dev/null
+++ b/Engine/Plugins/Experimental/GizmoEdMode/Source/GizmoEdMode/Public/DefaultAssetEditorGizmoFactory.h
@@ -0,0 +1,19 @@
+// Copyright Epic Games, Inc. All Rights Reserved.
+
+#pragma once
+
+#include "AssetEditorGizmoFactory.h"
+
+#include "DefaultAssetEditorGizmoFactory.generated.h"
+
+UCLASS()
+class GIZMOEDMODE_API UDefaultAssetEditorGizmoFactory : public UObject, public IAssetEditorGizmoFactory
+{
+	GENERATED_BODY()
+public:
+	//IAssetEditorGizmoFactory interface
+	virtual bool CanBuildGizmoForSelection(FEditorModeTools* ModeTools) const override;
+	virtual UTransformGizmo* BuildGizmoForSelection(FEditorModeTools* ModeTools, UInteractiveGizmoManager* GizmoManager) const override;
+	virtual EAssetEditorGizmoFactoryPriority GetPriority() const override { return EAssetEditorGizmoFactoryPriority::Default; }
+	virtual void ConfigureGridSnapping(bool bGridEnabled, bool bRotGridEnabled, UTransformGizmo* Gizmo) const override;
+};
diff --git a/Engine/Plugins/Experimental/GizmoEdMode/Source/GizmoEdMode/Public/GizmoEdMode.h b/Engine/Plugins/Experimental/GizmoEdMode/Source/GizmoEdMode/Public/GizmoEdMode.h
new file mode 100644
--- /dev/null
+++ b/Engine/Plugins/Experimental/GizmoEdMode/Source/GizmoEdMode/Public/GizmoEdMode.h
@@ -0,0 +1,51 @@
+// Copyright Epic Games, Inc. All Rights Reserved.
+
+#pragma once
+
+#include "Tools/UEdMode.h"
+#include "UObject/Object.h"
+#include "UObject/Interface.h"
+#include "UnrealWidget.h"
+#include "AssetEditorGizmoFactory.h"
+
+#include "GizmoEdMode.generated.h"
+
+class UTransformGizmo;
+class UInteractiveGizmoManager;
+
+UCLASS()
+class GIZMOEDMODE_API UGizmoEdModeSettings : public UObject
+{
+	GENERATED_BODY()
+};
+
+UCLASS()
+class GIZMOEDMODE_API UGizmoEdMode : public UEdMode
+{
+	GENERATED_BODY()
+public:
+	UGizmoEdMode();
+
+	void AddFactory(TScriptInterface<IAssetEditorGizmoFactory> GizmoFactory);
+
+private:
+	void ActorSelectionChangeNotify() override;
+	void Enter() override;
+	void Exit() override;
+	bool InputKey(FEditorViewportClient* ViewportClient, FViewport* Viewport, FKey Key, EInputEvent Event) override;
+
+	bool IsCompatibleWith(FEditorModeID OtherModeID) const override { return true; }
+	void Tick(FEditorViewportClient* ViewportClient, float DeltaTime) override;
+
+	void RecreateGizmo();
+	void DestroyGizmo();
+
+	UPROPERTY()
+	TArray<TScriptInterface<IAssetEditorGizmoFactory>> GizmoFactories;
+	IAssetEditorGizmoFactory* LastFactory = nullptr;
+
+	UPROPERTY()
+	UTransformGizmo* TransformGizmo = nullptr;
+
+	FDelegateHandle WidgetModeChangedHandle;
+};
diff --git a/Engine/Plugins/Experimental/GizmoEdMode/Source/GizmoEdMode/Public/GizmoEdModeModule.h b/Engine/Plugins/Experimental/GizmoEdMode/Source/GizmoEdMode/Public/GizmoEdModeModule.h
new file mode 100644
--- /dev/null
+++ b/Engine/Plugins/Experimental/GizmoEdMode/Source/GizmoEdMode/Public/GizmoEdModeModule.h
@@ -0,0 +1,18 @@
+// Copyright Epic Games, Inc. All Rights Reserved.
+
+#pragma once
+
+#include "CoreMinimal.h"
+#include "Modules/ModuleManager.h"
+
+class FGizmoEdModeModule : public IModuleInterface
+{
+public:
+
+	/** IModuleInterface implementation */
+	virtual void OnPostEngineInit();
+
+
+	virtual void StartupModule() override;
+	virtual void ShutdownModule() override;
+};
diff --git a/Engine/Plugins/Experimental/GizmoEdMode/Source/Private/DefaultAssetEditorGizmoFactory.cpp b/Engine/Plugins/Experimental/GizmoEdMode/Source/Private/DefaultAssetEditorGizmoFactory.cpp
deleted file mode 100644
diff --git a/Engine/Plugins/Experimental/GizmoEdMode/Source/Private/GizmoEdMode.cpp b/Engine/Plugins/Experimental/GizmoEdMode/Source/Private/GizmoEdMode.cpp
deleted file mode 100644
diff --git a/Engine/Plugins/Experimental/GizmoEdMode/Source/Private/GizmoEdModeModule.cpp b/Engine/Plugins/Experimental/GizmoEdMode/Source/Private/GizmoEdModeModule.cpp
deleted file mode 100644
diff --git a/Engine/Plugins/Experimental/GizmoEdMode/Source/Public/AssetEditorGizmoFactory.h b/Engine/Plugins/Experimental/GizmoEdMode/Source/Public/AssetEditorGizmoFactory.h
deleted file mode 100644
diff --git a/Engine/Plugins/Experimental/GizmoEdMode/Source/Public/DefaultAssetEditorGizmoFactory.h b/Engine/Plugins/Experimental/GizmoEdMode/Source/Public/DefaultAssetEditorGizmoFactory.h
deleted file mode 100644
diff --git a/Engine/Plugins/Experimental/GizmoEdMode/Source/Public/GizmoEdMode.h b/Engine/Plugins/Experimental/GizmoEdMode/Source/Public/GizmoEdMode.h
deleted file mode 100644
diff --git a/Engine/Plugins/Experimental/GizmoEdMode/Source/Public/GizmoEdModeModule.h b/Engine/Plugins/Experimental/GizmoEdMode/Source/Public/GizmoEdModeModule.h
deleted file mode 100644
