Allow new gizmo framework to support multiple gizmos per selection

#rnx

diff --git a/Engine/Plugins/Experimental/GizmoEdMode/Source/GizmoEdMode/Private/DefaultAssetEditorGizmoFactory.cpp b/Engine/Plugins/Experimental/GizmoEdMode/Source/GizmoEdMode/Private/DefaultAssetEditorGizmoFactory.cpp
--- a/Engine/Plugins/Experimental/GizmoEdMode/Source/GizmoEdMode/Private/DefaultAssetEditorGizmoFactory.cpp
+++ b/Engine/Plugins/Experimental/GizmoEdMode/Source/GizmoEdMode/Private/DefaultAssetEditorGizmoFactory.cpp
@@ -12,7 +12,7 @@ bool UDefaultAssetEditorGizmoFactory::CanBuildGizmoForSelection(FEditorModeTools
 	return true;
 }
 
-UTransformGizmo* UDefaultAssetEditorGizmoFactory::BuildGizmoForSelection(FEditorModeTools* ModeTools, UInteractiveGizmoManager* GizmoManager) const
+TArray<UInteractiveGizmo*> UDefaultAssetEditorGizmoFactory::BuildGizmoForSelection(FEditorModeTools* ModeTools, UInteractiveGizmoManager* GizmoManager) const
 {
 	ETransformGizmoSubElements Elements  = ETransformGizmoSubElements::None;
 	bool bUseContextCoordinateSystem = true;
@@ -50,15 +50,25 @@ UTransformGizmo* UDefaultAssetEditorGizmoFactory::BuildGizmoForSelection(FEditor
 	}
 	TransformGizmo->SetActiveTarget( TransformProxy );
 	TransformGizmo->SetVisibility(SelectedActors.Num() > 0);
-	return TransformGizmo;
+
+	TArray<UInteractiveGizmo*> Gizmos;
+
+	Gizmos.Add(TransformGizmo);
+
+	return Gizmos;
 }
 
-void UDefaultAssetEditorGizmoFactory::ConfigureGridSnapping(bool bGridEnabled, bool bRotGridEnabled, UTransformGizmo* Gizmo) const
+void UDefaultAssetEditorGizmoFactory::ConfigureGridSnapping(bool bGridEnabled, bool bRotGridEnabled, const TArray<UInteractiveGizmo*>& Gizmos) const
 {
-	if ( Gizmo )
+	for (auto& Gizmo : Gizmos)
 	{
-		Gizmo->bSnapToWorldGrid = bGridEnabled;
-		Gizmo->bSnapToWorldRotGrid = bRotGridEnabled;
+		UTransformGizmo* TransformGizmo = Cast<UTransformGizmo>(Gizmo);
+
+		if (TransformGizmo)
+		{
+			TransformGizmo->bSnapToWorldGrid = bGridEnabled;
+			TransformGizmo->bSnapToWorldRotGrid = bRotGridEnabled;
+		}
 	}
 
 }
diff --git a/Engine/Plugins/Experimental/GizmoEdMode/Source/GizmoEdMode/Private/GizmoEdMode.cpp b/Engine/Plugins/Experimental/GizmoEdMode/Source/GizmoEdMode/Private/GizmoEdMode.cpp
--- a/Engine/Plugins/Experimental/GizmoEdMode/Source/GizmoEdMode/Private/GizmoEdMode.cpp
+++ b/Engine/Plugins/Experimental/GizmoEdMode/Source/GizmoEdMode/Private/GizmoEdMode.cpp
@@ -49,7 +49,7 @@ void UGizmoEdMode::RecreateGizmo()
 	{
 		if (Factory->CanBuildGizmoForSelection(GetModeManager()))
 		{
-			TransformGizmo = Factory->BuildGizmoForSelection(GetModeManager(), ToolsContext->GizmoManager);
+			InteractiveGizmos = Factory->BuildGizmoForSelection(GetModeManager(), ToolsContext->GizmoManager);
 			LastFactory = &(*Factory);
 			return;
 		}
@@ -59,11 +59,13 @@ void UGizmoEdMode::RecreateGizmo()
 void UGizmoEdMode::DestroyGizmo()
 {
 	LastFactory = nullptr;
-	if ( TransformGizmo )
+
+	for (UInteractiveGizmo *Gizmo : InteractiveGizmos)
 	{
-		ToolsContext->GizmoManager->DestroyGizmo(TransformGizmo);
-		TransformGizmo = nullptr;
+		ToolsContext->GizmoManager->DestroyGizmo(Gizmo);
 	}
+
+	InteractiveGizmos.Empty();
 }
 
 void UGizmoEdMode::Enter()
@@ -95,7 +97,7 @@ void UGizmoEdMode::Tick(FEditorViewportClient* ViewportClient, float DeltaTime)
 	if ( LastFactory )
 	{
 		LastFactory->ConfigureGridSnapping(GetDefault<ULevelEditorViewportSettings>()->GridEnabled,
-		                                   GetDefault<ULevelEditorViewportSettings>()->RotGridEnabled, TransformGizmo);
+		                                   GetDefault<ULevelEditorViewportSettings>()->RotGridEnabled, InteractiveGizmos);
 	}
 }
 
diff --git a/Engine/Plugins/Experimental/GizmoEdMode/Source/GizmoEdMode/Public/AssetEditorGizmoFactory.h b/Engine/Plugins/Experimental/GizmoEdMode/Source/GizmoEdMode/Public/AssetEditorGizmoFactory.h
--- a/Engine/Plugins/Experimental/GizmoEdMode/Source/GizmoEdMode/Public/AssetEditorGizmoFactory.h
+++ b/Engine/Plugins/Experimental/GizmoEdMode/Source/GizmoEdMode/Public/AssetEditorGizmoFactory.h
@@ -9,7 +9,7 @@
 
 #include "AssetEditorGizmoFactory.generated.h"
 
-class UTransformGizmo;
+class UInteractiveGizmo;
 class UInteractiveGizmoManager;
 
 UENUM()
@@ -31,7 +31,7 @@ class GIZMOEDMODE_API IAssetEditorGizmoFactory
 	GENERATED_BODY()
 public:
 	virtual bool CanBuildGizmoForSelection(FEditorModeTools* ModeTools) const = 0;
-	virtual UTransformGizmo* BuildGizmoForSelection(FEditorModeTools* ModeTools, UInteractiveGizmoManager* GizmoManager) const = 0;
+	virtual TArray<UInteractiveGizmo*> BuildGizmoForSelection(FEditorModeTools* ModeTools, UInteractiveGizmoManager* GizmoManager) const = 0;
 	virtual EAssetEditorGizmoFactoryPriority GetPriority() const { return EAssetEditorGizmoFactoryPriority::Normal; }
-	virtual void ConfigureGridSnapping(bool bGridEnabled, bool bRotGridEnabled, UTransformGizmo* Gizmo) const = 0;
+	virtual void ConfigureGridSnapping(bool bGridEnabled, bool bRotGridEnabled, const TArray<UInteractiveGizmo*>& Gizmos) const = 0;
 };
diff --git a/Engine/Plugins/Experimental/GizmoEdMode/Source/GizmoEdMode/Public/DefaultAssetEditorGizmoFactory.h b/Engine/Plugins/Experimental/GizmoEdMode/Source/GizmoEdMode/Public/DefaultAssetEditorGizmoFactory.h
--- a/Engine/Plugins/Experimental/GizmoEdMode/Source/GizmoEdMode/Public/DefaultAssetEditorGizmoFactory.h
+++ b/Engine/Plugins/Experimental/GizmoEdMode/Source/GizmoEdMode/Public/DefaultAssetEditorGizmoFactory.h
@@ -13,7 +13,7 @@ class GIZMOEDMODE_API UDefaultAssetEditorGizmoFactory : public UObject, public I
 public:
 	//IAssetEditorGizmoFactory interface
 	virtual bool CanBuildGizmoForSelection(FEditorModeTools* ModeTools) const override;
-	virtual UTransformGizmo* BuildGizmoForSelection(FEditorModeTools* ModeTools, UInteractiveGizmoManager* GizmoManager) const override;
+	virtual TArray<UInteractiveGizmo*> BuildGizmoForSelection(FEditorModeTools* ModeTools, UInteractiveGizmoManager* GizmoManager) const override;
 	virtual EAssetEditorGizmoFactoryPriority GetPriority() const override { return EAssetEditorGizmoFactoryPriority::Default; }
-	virtual void ConfigureGridSnapping(bool bGridEnabled, bool bRotGridEnabled, UTransformGizmo* Gizmo) const override;
+	virtual void ConfigureGridSnapping(bool bGridEnabled, bool bRotGridEnabled, const TArray<UInteractiveGizmo*>& Gizmos) const override;
 };
diff --git a/Engine/Plugins/Experimental/GizmoEdMode/Source/GizmoEdMode/Public/GizmoEdMode.h b/Engine/Plugins/Experimental/GizmoEdMode/Source/GizmoEdMode/Public/GizmoEdMode.h
--- a/Engine/Plugins/Experimental/GizmoEdMode/Source/GizmoEdMode/Public/GizmoEdMode.h
+++ b/Engine/Plugins/Experimental/GizmoEdMode/Source/GizmoEdMode/Public/GizmoEdMode.h
@@ -11,6 +11,7 @@
 #include "GizmoEdMode.generated.h"
 
 class UTransformGizmo;
+class UInteractiveGizmo;
 class UInteractiveGizmoManager;
 
 UCLASS()
@@ -45,7 +46,7 @@ private:
 	IAssetEditorGizmoFactory* LastFactory = nullptr;
 
 	UPROPERTY()
-	UTransformGizmo* TransformGizmo = nullptr;
+	TArray<UInteractiveGizmo*> InteractiveGizmos;
 
 	FDelegateHandle WidgetModeChangedHandle;
 };
