Added initial prototype for the new in-viewport light gizmos using the Interactive Gizmo Framework.

diff --git a/Engine/Plugins/Experimental/GizmoEdMode/GizmoEdMode.uplugin b/Engine/Plugins/Experimental/GizmoEdMode/GizmoEdMode.uplugin
--- a/Engine/Plugins/Experimental/GizmoEdMode/GizmoEdMode.uplugin
+++ b/Engine/Plugins/Experimental/GizmoEdMode/GizmoEdMode.uplugin
@@ -19,6 +19,12 @@
 			"Name": "GizmoEdMode",
 			"Type": "Editor",
 			"LoadingPhase": "Default"
+		},
+		{
+			"Name": "LightGizmos",
+			"Type": "Editor",
+			"LoadingPhase": "Default"
+
 		}
 	]
 }
diff --git a/Engine/Plugins/Experimental/GizmoEdMode/Source/LightGizmos/LightGizmos.Build.cs b/Engine/Plugins/Experimental/GizmoEdMode/Source/LightGizmos/LightGizmos.Build.cs
new file mode 100644
--- /dev/null
+++ b/Engine/Plugins/Experimental/GizmoEdMode/Source/LightGizmos/LightGizmos.Build.cs
@@ -0,0 +1,57 @@
+// Copyright Epic Games, Inc. All Rights Reserved.
+
+using UnrealBuildTool;
+
+public class LightGizmos : ModuleRules
+{
+	public LightGizmos(ReadOnlyTargetRules Target) : base(Target)
+	{
+		PCHUsage = ModuleRules.PCHUsageMode.UseExplicitOrSharedPCHs;
+
+		PublicIncludePaths.AddRange(
+			new string[] {
+				// ... add public include paths required here ...
+			}
+			);
+
+
+		PrivateIncludePaths.AddRange(
+			new string[] {
+				// ... add other private include paths required here ...
+			}
+			);
+
+
+		PublicDependencyModuleNames.AddRange(
+			new string[]
+			{
+				// ... add other public dependencies that you statically link with here ...
+				"Core"
+			}
+			);
+
+
+		PrivateDependencyModuleNames.AddRange(
+			new string[]
+			{
+				"Engine",
+				"InputCore",
+				"UnrealEd",
+				"SlateCore",
+				"CoreUObject",
+				"EditorInteractiveToolsFramework",
+				"InteractiveToolsFramework",
+				"GizmoEdMode"
+				// ... add private dependencies that you statically link with here ...
+			}
+			);
+
+
+		DynamicallyLoadedModuleNames.AddRange(
+			new string[]
+			{
+				// ... add any modules that your module loads dynamically here ...
+			}
+			);
+	}
+}
diff --git a/Engine/Plugins/Experimental/GizmoEdMode/Source/LightGizmos/Private/DirectionalLightGizmo.cpp b/Engine/Plugins/Experimental/GizmoEdMode/Source/LightGizmos/Private/DirectionalLightGizmo.cpp
new file mode 100644
--- /dev/null
+++ b/Engine/Plugins/Experimental/GizmoEdMode/Source/LightGizmos/Private/DirectionalLightGizmo.cpp
@@ -0,0 +1,333 @@
+// Copyright Epic Games, Inc. All Rights Reserved.
+
+#include "DirectionalLightGizmo.h"
+#include "BaseGizmos/GizmoBoxComponent.h"
+#include "Components/SphereComponent.h"
+#include "Engine/CollisionProfile.h"
+#include "Kismet/KismetMathLibrary.h"
+#include "BaseGizmos/GizmoMath.h"
+#include "BaseGizmos/GizmoCircleComponent.h"
+#include "BaseGizmos/GizmoLineHandleComponent.h"
+
+// UDirectionalLightGizmoBuilder
+
+UInteractiveGizmo* UDirectionalLightGizmoBuilder::BuildGizmo(const FToolBuilderState& SceneState) const
+{
+	UDirectionalLightGizmo* NewGizmo = NewObject<UDirectionalLightGizmo>(SceneState.GizmoManager);
+	NewGizmo->SetWorld(SceneState.World);
+	return NewGizmo;
+}
+
+// ADirectionalLightGizmoActor
+
+ADirectionalLightGizmoActor::ADirectionalLightGizmoActor()
+{
+	// root component is a hidden sphere
+	USphereComponent* SphereComponent = CreateDefaultSubobject<USphereComponent>(TEXT("GizmoCenter"));
+	RootComponent = SphereComponent;
+	SphereComponent->InitSphereRadius(1.0f);
+	SphereComponent->SetVisibility(false);
+	SphereComponent->SetCollisionProfileName(UCollisionProfile::NoCollision_ProfileName);
+
+}
+
+// UDirectionalLightGizmo
+
+UDirectionalLightGizmo::UDirectionalLightGizmo()
+{
+	LightActor = nullptr;
+	GizmoActor = nullptr;
+	TransformProxy = nullptr;
+}
+
+void UDirectionalLightGizmo::Setup()
+{
+	UInteractiveGizmo::Setup();
+
+	UDirectionalLightGizmoInputBehavior* DirectionalLightBehavior = NewObject<UDirectionalLightGizmoInputBehavior>(this);
+	DirectionalLightBehavior->Initialize(this);
+	AddInputBehavior(DirectionalLightBehavior);
+
+	CreateGizmoHandles();
+
+	CreateZRotationGizmo();
+
+	// The Gizmo is being rotated around the Y axis
+	RotationPlaneX = FVector::XAxisVector;
+	RotationPlaneZ = FVector::ZAxisVector;
+	
+}
+
+void UDirectionalLightGizmo::Render(IToolsContextRenderAPI* RenderAPI)
+{
+	
+}
+
+void UDirectionalLightGizmo::Shutdown()
+{
+	if (GizmoActor)
+	{
+		GizmoActor->Destroy();
+		GizmoActor = nullptr;
+	}
+}
+
+USubTransformProxy* UDirectionalLightGizmo::GetTransformProxy()
+{
+	return TransformProxy;
+}
+
+void UDirectionalLightGizmo::SetSelectedObject(ADirectionalLight* InLight)
+{
+	LightActor = InLight;
+
+	// TODO: Cannot remove a component from Transform Proxy
+	if (!TransformProxy)
+	{
+		TransformProxy = NewObject<USubTransformProxy>(this);
+	}
+
+	USceneComponent* SceneComponent = LightActor->GetRootComponent();
+
+	TransformProxy->AddComponent(SceneComponent);
+
+	TransformProxy->OnTransformChanged.AddUObject(this, &UDirectionalLightGizmo::OnTransformChanged);
+
+	OnTransformChanged(TransformProxy, TransformProxy->GetTransform());
+
+}
+
+void UDirectionalLightGizmo::SetWorld(UWorld* InWorld)
+{
+	World = InWorld;
+}
+
+void UDirectionalLightGizmo::OnBeginDrag(const FInputDeviceRay& Ray)
+{
+	FVector Start = Ray.WorldRay.Origin;
+	const float MaxRaycastDistance = 1e6f;
+	FVector End = Ray.WorldRay.Origin + Ray.WorldRay.Direction * MaxRaycastDistance;
+
+	FRay HitCheckRay(Start, End - Start);
+	FHitResult HitResult;
+
+	FTransform DragTransform;
+	HitComponent = nullptr;
+
+	// Check if any component was hit
+	if (HitTest(HitCheckRay, HitResult, DragTransform, HitComponent))
+	{
+		// Rotate around y axis if the arrow was hit
+		if (HitComponent == GizmoActor->Arrow)
+		{
+			HitAxis = FVector::YAxisVector;
+
+			// Get the rotated plane vectors for the interaction
+			GizmoMath::MakeNormalPlaneBasis(HitAxis, RotationPlaneX, RotationPlaneZ);
+			RotationPlaneX = LightActor->GetActorRotation().RotateVector(FVector::XAxisVector);
+			RotationPlaneZ = LightActor->GetActorRotation().RotateVector(FVector::ZAxisVector);
+
+		}
+		// Rotate around Z axis if the circle was hit
+		else
+		{
+			HitAxis = FVector::ZAxisVector;
+			RotationPlaneX = FVector::XAxisVector;
+			RotationPlaneZ = FVector::YAxisVector;
+		}
+
+		// Calculate initial hit position
+		DragStartWorldPosition = GizmoMath::ProjectPointOntoLine(Ray.WorldRay.PointAt(HitResult.Distance), DragTransform.GetLocation(), HitAxis);
+
+		// Calculate initial hit parameters
+		bool bIntersects; 
+		FVector IntersectionPoint;
+
+		GizmoMath::RayPlaneIntersectionPoint(
+			DragStartWorldPosition, HitAxis,
+			Ray.WorldRay.Origin, Ray.WorldRay.Direction,
+			bIntersects, IntersectionPoint);
+
+		check(bIntersects);  // need to handle this case...
+
+		InteractionStartPoint = IntersectionPoint;
+
+		InteractionStartParameter = GizmoMath::ComputeAngleInPlane(InteractionStartPoint,
+			DragStartWorldPosition, HitAxis, RotationPlaneX, RotationPlaneZ);
+	}
+}
+
+void UDirectionalLightGizmo::OnUpdateDrag(const FInputDeviceRay& Ray)
+{
+	bool bIntersects; 
+	FVector IntersectionPoint;
+
+	// Calculate current hit parameters
+	GizmoMath::RayPlaneIntersectionPoint(
+		DragStartWorldPosition, HitAxis,
+		Ray.WorldRay.Origin, Ray.WorldRay.Direction,
+		bIntersects, IntersectionPoint);
+
+	if (!bIntersects)
+	{
+		return;
+	}
+
+	FVector InteractionCurPoint = IntersectionPoint;
+
+	float InteractionCurAngle = GizmoMath::ComputeAngleInPlane(InteractionCurPoint,
+		DragStartWorldPosition, HitAxis, RotationPlaneX, RotationPlaneZ);
+
+	float DeltaAngle = InteractionCurAngle - InteractionStartParameter;
+
+	// Rotate around y axis if the arrow was hit
+	if (HitComponent == GizmoActor->Arrow)
+	{
+		FRotator Rotation = FRotator::ZeroRotator;
+
+		Rotation.Pitch = FMath::RadiansToDegrees(DeltaAngle);
+
+		LightActor->AddActorLocalRotation(Rotation);
+	}
+	// Rotate around Z axis if the circle was hit
+	else
+	{
+		FQuat Rotation(FVector(0, 0, 1), DeltaAngle);
+		LightActor->AddActorWorldRotation(Rotation);
+	}
+
+	TransformProxy->SetTransform(LightActor->GetTransform());
+
+	InteractionStartPoint = InteractionCurPoint;
+	InteractionStartParameter = InteractionCurAngle;
+
+}
+
+bool UDirectionalLightGizmo::HitTest(const FRay& Ray, FHitResult& OutHit, FTransform& OutTransform, UPrimitiveComponent*& OutHitComponent)
+{
+	FVector Start = Ray.Origin;
+	const float MaxRaycastDistance = 1e6f;
+	FVector End = Ray.Origin + Ray.Direction * MaxRaycastDistance;
+
+	FCollisionQueryParams Params;
+
+	if (GizmoActor->Arrow->LineTraceComponent(OutHit, Start, End, Params))
+	{
+		OutTransform = GizmoActor->Arrow->GetComponentTransform();
+		OutHitComponent = GizmoActor->Arrow;
+		return true;
+	}
+	else if (GizmoActor->RotationZCircle && GizmoActor->RotationZCircle->LineTraceComponent(OutHit, Start, End, Params))
+	{
+		OutTransform = GizmoActor->GetTransform();
+		OutHitComponent = GizmoActor->RotationZCircle;
+		return true;
+	}
+
+	return false;
+}
+
+void UDirectionalLightGizmo::CreateGizmoHandles()
+{
+	FActorSpawnParameters SpawnInfo;
+	GizmoActor = World->SpawnActor<ADirectionalLightGizmoActor>(FVector::ZeroVector, FRotator::ZeroRotator, SpawnInfo);
+
+	GizmoActor->Arrow = AGizmoActor::AddDefaultLineHandleComponent(World, GizmoActor, FLinearColor::Red, FVector::YAxisVector, FVector::XAxisVector, ArrowLength, true);
+}
+
+void UDirectionalLightGizmo::UpdateGizmoHandles()
+{
+	if (GizmoActor && GizmoActor->RotationZCircle)
+	{
+		GizmoActor->RotationZCircle->SetRelativeRotation(GizmoActor->GetActorRotation().Quaternion().Inverse());
+	}
+}
+
+void UDirectionalLightGizmo::OnTransformChanged(UTransformProxy*, FTransform)
+{
+	USceneComponent* GizmoComponent = GizmoActor->GetRootComponent();
+
+	FTransform TargetTransform = TransformProxy->GetTransform();
+
+	FTransform GizmoTransform = TargetTransform;
+	GizmoTransform.SetScale3D(FVector(1, 1, 1));
+
+	GizmoComponent->SetWorldTransform(GizmoTransform);
+
+	UpdateGizmoHandles();
+}
+
+
+void UDirectionalLightGizmo::CreateZRotationGizmo()
+{
+	UGizmoCircleComponent* NewCircle = NewObject<UGizmoCircleComponent>(GizmoActor);
+	GizmoActor->AddInstanceComponent(NewCircle);
+	NewCircle->AttachToComponent(GizmoActor->GetRootComponent(), FAttachmentTransformRules::KeepRelativeTransform);
+	NewCircle->Normal = FVector::ZAxisVector;
+	NewCircle->Color = FLinearColor::Blue;
+	NewCircle->Radius = 120.f;
+	NewCircle->RegisterComponent();
+
+	GizmoActor->RotationZCircle = NewCircle;
+
+	UpdateGizmoHandles();
+}
+
+// UDirectionalLightGizmoInputBehavior
+
+void UDirectionalLightGizmoInputBehavior::Initialize(UDirectionalLightGizmo* InGizmo)
+{
+	Gizmo = InGizmo;
+}
+
+FInputCaptureRequest UDirectionalLightGizmoInputBehavior::WantsCapture(const FInputDeviceState& input)
+{
+	if (IsPressed(input))
+	{
+		FHitResult HitResult;
+		FTransform DragTransform;
+		UPrimitiveComponent* HitComponent = nullptr;
+		if (Gizmo->HitTest(input.Mouse.WorldRay, HitResult, DragTransform, HitComponent))
+		{
+			return FInputCaptureRequest::Begin(this, EInputCaptureSide::Any, HitResult.Distance);
+		}
+	}
+
+	return FInputCaptureRequest::Ignore();
+}
+
+FInputCaptureUpdate UDirectionalLightGizmoInputBehavior::BeginCapture(const FInputDeviceState& input, EInputCaptureSide eSide)
+{
+	FInputDeviceRay DeviceRay(input.Mouse.WorldRay, input.Mouse.Position2D);
+	LastWorldRay = DeviceRay.WorldRay;
+	LastScreenPosition = DeviceRay.ScreenPosition;
+
+	Gizmo->OnBeginDrag(DeviceRay);
+	bInputDragCaptured = true;
+	return FInputCaptureUpdate::Begin(this, EInputCaptureSide::Any);
+}
+
+FInputCaptureUpdate UDirectionalLightGizmoInputBehavior::UpdateCapture(const FInputDeviceState& input, const FInputCaptureData& data)
+{
+	FInputDeviceRay DeviceRay(input.Mouse.WorldRay, input.Mouse.Position2D);
+	LastWorldRay = DeviceRay.WorldRay;
+	LastScreenPosition = DeviceRay.ScreenPosition;
+
+	if (IsReleased(input))
+	{
+		bInputDragCaptured = false;
+		return FInputCaptureUpdate::End();
+	}
+
+	Gizmo->OnUpdateDrag(LastWorldRay);
+
+	return FInputCaptureUpdate::Continue();
+}
+
+void UDirectionalLightGizmoInputBehavior::ForceEndCapture(const FInputCaptureData& data)
+{
+	if (bInputDragCaptured)
+	{
+		bInputDragCaptured = false;
+	}
+}
\ No newline at end of file
diff --git a/Engine/Plugins/Experimental/GizmoEdMode/Source/LightGizmos/Private/DirectionalLightGizmo.h b/Engine/Plugins/Experimental/GizmoEdMode/Source/LightGizmos/Private/DirectionalLightGizmo.h
new file mode 100644
--- /dev/null
+++ b/Engine/Plugins/Experimental/GizmoEdMode/Source/LightGizmos/Private/DirectionalLightGizmo.h
@@ -0,0 +1,157 @@
+// Copyright Epic Games, Inc. All Rights Reserved.
+
+#pragma once
+
+#include "CoreMinimal.h"
+#include "InteractiveGizmo.h"
+#include "InteractiveGizmoBuilder.h"
+#include "SubTransformProxy.h"
+#include "BaseGizmos/GizmoActor.h"
+#include "BaseBehaviors/BehaviorTargetInterfaces.h"
+#include "BaseBehaviors/AnyButtonInputBehavior.h"
+#include "Engine/DirectionalLight.h"
+#include "BaseGizmos/TransformGizmo.h"
+#include "SubTransformProxy.h"
+
+#include "DirectionalLightGizmo.generated.h"
+
+UCLASS()
+class UDirectionalLightGizmoBuilder : public UInteractiveGizmoBuilder
+{
+	GENERATED_BODY()
+
+public:
+	virtual UInteractiveGizmo* BuildGizmo(const FToolBuilderState& SceneState) const override;
+};
+
+UCLASS()
+class ADirectionalLightGizmoActor : public AGizmoActor
+{
+	GENERATED_BODY()
+
+public:
+
+	ADirectionalLightGizmoActor();
+
+	// The handle to rotate around its y axis
+	UPrimitiveComponent* Arrow;
+
+	// The handle to rotate around the world z axis
+	UPrimitiveComponent* RotationZCircle;
+};
+
+/**
+ * UDirectionalLightGizmo provides a gizmo to allow for editing directional light properties in viewport
+ * Currently supports rotating the light around the world Z axis and its Y Axis
+ *
+ */
+UCLASS()
+class UDirectionalLightGizmo : public UInteractiveGizmo, public IHoverBehaviorTarget
+{
+	GENERATED_BODY()
+
+public:
+	// UInteractiveGizmo interface
+
+	UDirectionalLightGizmo();
+
+	virtual void Setup() override;
+
+	virtual void Render(IToolsContextRenderAPI* RenderAPI) override;
+
+	virtual void Shutdown() override;
+
+	USubTransformProxy* GetTransformProxy();
+
+	// IHoverBehaviorTarget interface
+	virtual FInputRayHit BeginHoverSequenceHitTest(const FInputDeviceRay& PressPos) override { return FInputRayHit(); }
+	virtual void OnBeginHover(const FInputDeviceRay& DevicePos) override {}
+	virtual bool OnUpdateHover(const FInputDeviceRay& DevicePos) override { return true; }
+	virtual void OnEndHover() override {}
+
+	virtual void SetSelectedObject(ADirectionalLight* InLight);
+	virtual void SetWorld(UWorld* InWorld);
+
+	virtual void OnBeginDrag(const FInputDeviceRay& Ray);
+	virtual void OnUpdateDrag(const FInputDeviceRay& Ray);
+
+	/** Check if any of the components are hit by the input ray */
+	bool HitTest(const FRay& Ray, FHitResult& OutHit, FTransform& OutTransform, UPrimitiveComponent*& OutHitComponent);
+
+private:
+
+	void CreateGizmoHandles();
+	void UpdateGizmoHandles();
+
+	void OnTransformChanged(UTransformProxy*, FTransform);
+
+	/** Create the gizmo that rotates around the world z axis */
+	void CreateZRotationGizmo();
+
+	/** A transform proxy of the LightActor */
+	UPROPERTY()
+	USubTransformProxy* TransformProxy;
+
+	UPROPERTY()
+	UWorld* World;
+
+	/** The internal actor used by the light gizmo */
+	UPROPERTY()
+	ADirectionalLightGizmoActor* GizmoActor;
+
+	/** The current target of the gizmo */
+	UPROPERTY()
+	ADirectionalLight* LightActor;
+
+	/** Parameters used during hit testing */
+	UPROPERTY()
+	FVector DragStartWorldPosition;
+
+	UPROPERTY()
+	FVector InteractionStartPoint;
+
+	UPROPERTY()
+	float InteractionStartParameter;
+
+	UPROPERTY()
+	FVector HitAxis;
+
+	UPROPERTY()
+	FVector RotationPlaneX;
+
+	UPROPERTY()
+	FVector RotationPlaneZ;
+
+	UPROPERTY()
+	UPrimitiveComponent* HitComponent;
+
+	UPROPERTY()
+	float ArrowLength{ 120.f };
+};
+
+
+/**
+ * A behavior that forwards clicking and dragging to the gizmo.
+ */
+UCLASS()
+class UDirectionalLightGizmoInputBehavior : public UAnyButtonInputBehavior
+{
+	GENERATED_BODY()
+
+public:
+	virtual FInputCapturePriority GetPriority() override { return FInputCapturePriority(FInputCapturePriority::DEFAULT_GIZMO_PRIORITY); }
+
+	virtual void Initialize(UDirectionalLightGizmo* Gizmo);
+
+	virtual FInputCaptureRequest WantsCapture(const FInputDeviceState& input) override;
+	virtual FInputCaptureUpdate BeginCapture(const FInputDeviceState& input, EInputCaptureSide eSide) override;
+	virtual FInputCaptureUpdate UpdateCapture(const FInputDeviceState& input, const FInputCaptureData& data) override;
+	virtual void ForceEndCapture(const FInputCaptureData& data) override;
+
+protected:
+	UDirectionalLightGizmo* Gizmo;
+	FRay LastWorldRay;
+	FVector2D LastScreenPosition;
+	bool bInputDragCaptured;
+
+};
\ No newline at end of file
diff --git a/Engine/Plugins/Experimental/GizmoEdMode/Source/LightGizmos/Private/DirectionalLightGizmoFactory.cpp b/Engine/Plugins/Experimental/GizmoEdMode/Source/LightGizmos/Private/DirectionalLightGizmoFactory.cpp
new file mode 100644
--- /dev/null
+++ b/Engine/Plugins/Experimental/GizmoEdMode/Source/LightGizmos/Private/DirectionalLightGizmoFactory.cpp
@@ -0,0 +1,115 @@
+// Copyright Epic Games, Inc. All Rights Reserved.
+
+#include "DirectionalLightGizmoFactory.h"
+#include "EditorModeManager.h"
+#include "Engine/Selection.h"
+#include "BaseGizmos/TransformGizmo.h"
+#include "BaseGizmos/TransformProxy.h"
+#include "Math/Rotator.h"
+#include "UnrealWidget.h"
+#include "Engine/DirectionalLight.h"
+#include "DirectionalLightGizmo.h"
+#include "LightGizmosModule.h"
+#include "SubTransformProxy.h"
+
+bool UDirectionalLightGizmoFactory::CanBuildGizmoForSelection(FEditorModeTools* ModeTools) const
+{
+	/** Since there can only be one factory active at a time in the current iteration
+	 *  of the tools framework, this function only returns true if /all/ items in the
+	 *  current selection are point lights and then creates point light gizmos for them
+	 *  If not, the selection falls through to the next factory
+	 */
+	TArray<AActor*> SelectedActors;
+	ModeTools->GetSelectedActors()->GetSelectedObjects<AActor>(SelectedActors);
+	
+	for (const AActor* Actor : SelectedActors)
+	{
+		if (!Cast<ADirectionalLight>(Actor))
+		{
+			return false;
+		}
+	}
+
+	// Only return true if all selected Actors are Directional Lights
+	return true;
+}
+
+TArray<UInteractiveGizmo*> UDirectionalLightGizmoFactory::BuildGizmoForSelection(FEditorModeTools* ModeTools, UInteractiveGizmoManager* GizmoManager) const
+{
+	TArray<UInteractiveGizmo*> Gizmos;
+	TArray<AActor*> SelectedActors;
+	ModeTools->GetSelectedActors()->GetSelectedObjects<AActor>(SelectedActors);
+
+	// Createa Light Gizmo for each actor in selection
+	for (AActor* Actor : SelectedActors)
+	{
+		UDirectionalLightGizmo* LightGizmo = Cast<UDirectionalLightGizmo>(GizmoManager->CreateGizmo(FLightGizmosModule::DirectionalLightGizmoType));
+
+		LightGizmo->SetSelectedObject(Cast<ADirectionalLight>(Actor));
+
+		Gizmos.Add(LightGizmo);
+	}
+
+	// Create one transform gizmo for the whole selection
+
+	ETransformGizmoSubElements Elements = ETransformGizmoSubElements::None;
+	bool bUseContextCoordinateSystem = true;
+	FWidget::EWidgetMode WidgetMode = ModeTools->GetWidgetMode();
+	switch (WidgetMode)
+	{
+	case FWidget::EWidgetMode::WM_Translate:
+		Elements = ETransformGizmoSubElements::TranslateAllAxes | ETransformGizmoSubElements::TranslateAllPlanes;
+		break;
+	case FWidget::EWidgetMode::WM_Rotate:
+		Elements = ETransformGizmoSubElements::None; // Custom Rotation in DirectionalLightGizmo
+		break;
+	case FWidget::EWidgetMode::WM_Scale:
+		Elements = ETransformGizmoSubElements::None; // Scaling doesn't make sense for a directional light
+		bUseContextCoordinateSystem = false;
+		break;
+	case FWidget::EWidgetMode::WM_2D:
+		Elements = ETransformGizmoSubElements::RotateAxisY | ETransformGizmoSubElements::TranslatePlaneXZ;
+		break;
+	default:
+		Elements = ETransformGizmoSubElements::FullTranslateRotateScale;
+		break;
+	}
+	UTransformGizmo* TransformGizmo = GizmoManager->CreateCustomTransformGizmo(Elements);
+	TransformGizmo->bUseContextCoordinateSystem = bUseContextCoordinateSystem;
+
+	USubTransformProxy* TransformProxy = NewObject<USubTransformProxy>();
+
+	// Get the SubTransformProxy from each UDirectionalLightGizmo and connect it to the main Transform Proxy
+	for (UInteractiveGizmo* Gizmo : Gizmos)
+	{
+		UDirectionalLightGizmo* LightGizmo = Cast<UDirectionalLightGizmo>(Gizmo);
+
+		if (LightGizmo)
+		{
+			TransformProxy->AddSubTransformProxy(LightGizmo->GetTransformProxy());
+		}
+
+	}
+
+	TransformGizmo->SetActiveTarget(TransformProxy);
+	TransformGizmo->SetVisibility(SelectedActors.Num() > 0);
+
+	Gizmos.Add(TransformGizmo);
+	
+	return Gizmos;
+
+}
+
+void UDirectionalLightGizmoFactory::ConfigureGridSnapping(bool bGridEnabled, bool bRotGridEnabled, const TArray<UInteractiveGizmo*>& Gizmos) const
+{
+	for (UInteractiveGizmo* Gizmo : Gizmos)
+	{
+		UTransformGizmo* TransformGizmo = Cast<UTransformGizmo>(Gizmo);
+
+		if (TransformGizmo)
+		{
+			TransformGizmo->bSnapToWorldGrid = bGridEnabled;
+			TransformGizmo->bSnapToWorldRotGrid = bRotGridEnabled;
+		}
+	}
+}
diff --git a/Engine/Plugins/Experimental/GizmoEdMode/Source/LightGizmos/Private/LightGizmosModule.cpp b/Engine/Plugins/Experimental/GizmoEdMode/Source/LightGizmos/Private/LightGizmosModule.cpp
new file mode 100644
--- /dev/null
+++ b/Engine/Plugins/Experimental/GizmoEdMode/Source/LightGizmos/Private/LightGizmosModule.cpp
@@ -0,0 +1,71 @@
+// Copyright Epic Games, Inc. All Rights Reserved.
+
+#include "LightGizmosModule.h"
+#include "Editor.h"
+#include "EditorModeRegistry.h"
+#include "EditorModeManager.h"
+#include "Misc/CoreDelegates.h"
+#include "GizmoEdMode.h"
+#include "PointLightGizmoFactory.h"
+#include "DirectionalLightGizmoFactory.h"
+#include "SpotLightGizmoFactory.h"
+#include "PointLightGizmo.h"
+#include "ScalableConeGizmo.h"
+#include "SpotLightGizmo.h"
+#include "DirectionalLightGizmo.h"
+
+#define LOCTEXT_NAMESPACE "FLightGizmosModule"
+
+FString FLightGizmosModule::PointLightGizmoType = TEXT("PointLightGizmoType");
+FString FLightGizmosModule::SpotLightGizmoType = TEXT("SpotLightGizmoType");
+FString FLightGizmosModule::ScalableConeGizmoType = TEXT("ScalableConeGizmoType");
+FString FLightGizmosModule::DirectionalLightGizmoType = TEXT("DirectionalLightGizmoType");
+
+void FLightGizmosModule::OnPostEngineInit()
+{
+	GLevelEditorModeTools().OnEditorModeIDChanged().AddLambda([](const FEditorModeID& ModeID, bool IsEnteringMode) 
+	{
+		if ((ModeID == GetDefault<UGizmoEdMode>()->GetID()))
+		{
+			UEdMode* EdMode = GLevelEditorModeTools().GetActiveScriptableMode(GetDefault<UGizmoEdMode>()->GetID());
+			UGizmoEdMode* GizmoEdMode = Cast<UGizmoEdMode>(EdMode);
+
+			 // Register the factories and gizmos if we are entering the new Gizmo Mode
+			if (IsEnteringMode)
+			{
+				GizmoEdMode->GetToolManager()->GetPairedGizmoManager()->RegisterGizmoType(PointLightGizmoType, NewObject<UPointLightGizmoBuilder>());
+				GizmoEdMode->GetToolManager()->GetPairedGizmoManager()->RegisterGizmoType(ScalableConeGizmoType, NewObject<UScalableConeGizmoBuilder>());
+				GizmoEdMode->GetToolManager()->GetPairedGizmoManager()->RegisterGizmoType(SpotLightGizmoType, NewObject<USpotLightGizmoBuilder>());
+				GizmoEdMode->GetToolManager()->GetPairedGizmoManager()->RegisterGizmoType(DirectionalLightGizmoType, NewObject<UDirectionalLightGizmoBuilder>());
+				GizmoEdMode->AddFactory(NewObject<UPointLightGizmoFactory>());
+				GizmoEdMode->AddFactory(NewObject<UDirectionalLightGizmoFactory>());
+				GizmoEdMode->AddFactory(NewObject<USpotLightGizmoFactory>());
+			}
+			else
+			{
+				//TODO: Deregistering on exit currently causes a crash
+
+				//GizmoEdMode->GetToolManager()->GetPairedGizmoManager()->DeregisterGizmoType(PointLightGizmoType);
+				//GizmoEdMode->GetToolManager()->GetPairedGizmoManager()->DeregisterGizmoType(ScalableConeGizmoType);
+				//GizmoEdMode->GetToolManager()->GetPairedGizmoManager()->DeregisterGizmoType(SpotLightGizmoType);
+				//GizmoEdMode->GetToolManager()->GetPairedGizmoManager()->DeregisterGizmoType(DirectionalLightGizmoType);
+			}
+			
+		}
+	});
+
+}
+
+void FLightGizmosModule::StartupModule()
+{
+	FCoreDelegates::OnPostEngineInit.AddRaw(this, &FLightGizmosModule::OnPostEngineInit);
+}
+
+void FLightGizmosModule::ShutdownModule()
+{
+
+}
+
+#undef LOCTEXT_NAMESPACE
+
+IMPLEMENT_MODULE(FLightGizmosModule, LightGizmos);
diff --git a/Engine/Plugins/Experimental/GizmoEdMode/Source/LightGizmos/Private/PointLightGizmo.cpp b/Engine/Plugins/Experimental/GizmoEdMode/Source/LightGizmos/Private/PointLightGizmo.cpp
new file mode 100644
--- /dev/null
+++ b/Engine/Plugins/Experimental/GizmoEdMode/Source/LightGizmos/Private/PointLightGizmo.cpp
@@ -0,0 +1,92 @@
+// Copyright Epic Games, Inc. All Rights Reserved.
+
+#include "PointLightGizmo.h"
+#include "Components/PointLightComponent.h"
+
+UInteractiveGizmo* UPointLightGizmoBuilder::BuildGizmo(const FToolBuilderState& SceneState) const
+{
+	UPointLightGizmo* NewGizmo = NewObject<UPointLightGizmo>(SceneState.GizmoManager);
+	NewGizmo->SetWorld(SceneState.World);
+	return NewGizmo;
+}
+
+UPointLightGizmo::UPointLightGizmo()
+{
+	LightActor = nullptr;
+	AttenuationGizmo = nullptr;
+	TransformProxy = nullptr;
+}
+
+void UPointLightGizmo::Setup()
+{
+	
+}
+
+void UPointLightGizmo::Tick(float DeltaTime)
+{
+	 // Make sure that radius of the sphere is correct every frame
+	AttenuationGizmo->SetRadius(LightActor->PointLightComponent->AttenuationRadius);
+}
+
+void UPointLightGizmo::Render(IToolsContextRenderAPI* RenderAPI)
+{
+
+}
+
+void UPointLightGizmo::Shutdown()
+{
+	if (AttenuationGizmo)
+	{
+		GetGizmoManager()->DestroyGizmo(AttenuationGizmo);
+		AttenuationGizmo = nullptr;
+	}
+}
+
+void UPointLightGizmo::SetSelectedObject(APointLight* InLight)
+{
+	LightActor = InLight;
+
+	// Create a transform proxy and add the light to it
+	// TODO: Cannot change selected object right now as there is no way to remove a component from the transform proxy
+	if (!TransformProxy)
+	{
+		TransformProxy = NewObject<USubTransformProxy>(this);
+	}
+
+	USceneComponent* SceneComponent = LightActor->GetRootComponent();
+	TransformProxy->AddComponent(SceneComponent);
+}
+
+USubTransformProxy* UPointLightGizmo::GetTransformProxy()
+{
+	return TransformProxy;
+}
+
+
+void UPointLightGizmo::CreateLightGizmo()
+{
+	if (!LightActor)
+	{
+		return;
+	}
+
+	AttenuationGizmo = GetGizmoManager()->CreateGizmo<UScalableSphereGizmo>(UInteractiveGizmoManager::DefaultScalableSphereBuilderIdentifier);
+
+	AttenuationGizmo->SetTarget(TransformProxy);
+
+	AttenuationGizmo->UpdateRadiusFunc = [this](float NewRadius) { this->OnAttenuationUpdate(NewRadius); };
+}
+
+void UPointLightGizmo::SetWorld(UWorld* InWorld)
+{
+	World = InWorld;
+}
+
+void UPointLightGizmo::OnAttenuationUpdate(float NewRadius)
+{
+	// Update the attenuation radius
+	// TODO: Is this the right way to update radius?
+	LightActor->PointLightComponent->AttenuationRadius = NewRadius;
+	LightActor->PointLightComponent->MarkRenderStateDirty();
+}
+
diff --git a/Engine/Plugins/Experimental/GizmoEdMode/Source/LightGizmos/Private/PointLightGizmo.h b/Engine/Plugins/Experimental/GizmoEdMode/Source/LightGizmos/Private/PointLightGizmo.h
new file mode 100644
--- /dev/null
+++ b/Engine/Plugins/Experimental/GizmoEdMode/Source/LightGizmos/Private/PointLightGizmo.h
@@ -0,0 +1,80 @@
+// Copyright Epic Games, Inc. All Rights Reserved.
+
+#pragma once
+
+#include "BaseGizmos/ScalableSphereGizmo.h"
+#include "Engine/PointLight.h"
+#include "InteractiveGizmo.h"
+#include "SubTransformProxy.h"
+
+#include "PointLightGizmo.generated.h"
+
+UCLASS()
+class UPointLightGizmoBuilder : public UInteractiveGizmoBuilder
+{
+	GENERATED_BODY()
+
+public:
+	virtual UInteractiveGizmo* BuildGizmo(const FToolBuilderState& SceneState) const override;
+};
+
+/**
+ * UPointLightGizmo provides a gizmo to allow for editing point light properties in viewport
+ * Currently supports changing the attenuation radius using a USacalableSphereGizmo
+ *
+ */
+UCLASS()
+class UPointLightGizmo : public UInteractiveGizmo
+{
+	GENERATED_BODY()
+
+public:
+	// UInteractiveGizmo interface
+
+	virtual void Setup() override;
+
+	virtual void Tick(float DeltaTime) override;
+
+	virtual void Render(IToolsContextRenderAPI* RenderAPI) override;
+
+	virtual void Shutdown() override;
+
+	UPointLightGizmo();
+
+	/**
+	 * Set the Target to which the gizmo will be attached
+	 */
+	void SetSelectedObject(APointLight*InLight);
+	
+	/**
+	 * Create the Attenuation Edit Gizmo
+	 */
+	void CreateLightGizmo();
+
+	void SetWorld(UWorld *InWorld);
+
+	USubTransformProxy* GetTransformProxy();
+
+private:
+
+	/** The Gizmo that is used to scale the attenuation */
+	UPROPERTY()
+	UScalableSphereGizmo* AttenuationGizmo;
+
+	/** The target Point Light */
+	UPROPERTY()
+	APointLight* LightActor;
+
+	UPROPERTY()
+	UWorld* World;
+
+	/** A transform proxy to represent the light actor in gizmos */
+	UPROPERTY()
+	USubTransformProxy* TransformProxy;
+	
+	/**
+	 * Function that is called when the radius of the sphere is changed
+	 */
+	void OnAttenuationUpdate(float NewRadius);
+
+};
\ No newline at end of file
diff --git a/Engine/Plugins/Experimental/GizmoEdMode/Source/LightGizmos/Private/PointLightGizmoFactory.cpp b/Engine/Plugins/Experimental/GizmoEdMode/Source/LightGizmos/Private/PointLightGizmoFactory.cpp
new file mode 100644
--- /dev/null
+++ b/Engine/Plugins/Experimental/GizmoEdMode/Source/LightGizmos/Private/PointLightGizmoFactory.cpp
@@ -0,0 +1,116 @@
+// Copyright Epic Games, Inc. All Rights Reserved.
+
+#include "PointLightGizmoFactory.h"
+#include "EditorModeManager.h"
+#include "Engine/Selection.h"
+#include "BaseGizmos/TransformGizmo.h"
+#include "SubTransformProxy.h"
+#include "Math/Rotator.h"
+#include "UnrealWidget.h"
+#include "Engine/PointLight.h"
+#include "PointLightGizmo.h"
+#include "LightGizmosModule.h"
+
+bool UPointLightGizmoFactory::CanBuildGizmoForSelection(FEditorModeTools* ModeTools) const
+{
+	/** Since there can only be one factory active at a time in the current iteration
+	 *  of the tools framework, this function only returns true if /all/ items in the
+	 *  current selection are point lights and then creates point light gizmos for them
+	 *  If not, the selection falls through to the next factory
+	 */
+	TArray<AActor*> SelectedActors;
+	ModeTools->GetSelectedActors()->GetSelectedObjects<AActor>(SelectedActors);
+
+	for (const AActor* Actor : SelectedActors)
+	{
+		if (!Cast<APointLight>(Actor))
+		{
+			return false;
+		}
+	}
+
+	// Only return true if all selected actors are Point Light
+	return true;
+}
+
+TArray<UInteractiveGizmo*> UPointLightGizmoFactory::BuildGizmoForSelection(FEditorModeTools* ModeTools, UInteractiveGizmoManager* GizmoManager) const
+{
+	TArray<UInteractiveGizmo*> Gizmos;
+	TArray<AActor*> SelectedActors;
+	ModeTools->GetSelectedActors()->GetSelectedObjects<AActor>(SelectedActors);
+	
+	 // Create a point light gizmo for each actor
+	for (AActor* Actor : SelectedActors)
+	{
+		UPointLightGizmo* LightGizmo = Cast<UPointLightGizmo>(GizmoManager->CreateGizmo(FLightGizmosModule::PointLightGizmoType));
+
+		LightGizmo->SetSelectedObject(Cast<APointLight>(Actor));
+		LightGizmo->CreateLightGizmo();
+
+		Gizmos.Add(LightGizmo);
+	}
+
+
+	// Create one transform gizmo for the whole selection
+	ETransformGizmoSubElements Elements = ETransformGizmoSubElements::None;
+	bool bUseContextCoordinateSystem = true;
+	FWidget::EWidgetMode WidgetMode = ModeTools->GetWidgetMode();
+
+	switch (WidgetMode)
+	{
+	case FWidget::EWidgetMode::WM_Translate:
+		Elements = ETransformGizmoSubElements::TranslateAllAxes | ETransformGizmoSubElements::TranslateAllPlanes;
+		break;
+	case FWidget::EWidgetMode::WM_Rotate:
+		Elements = ETransformGizmoSubElements::None; // Rotation doesn't make sense for point lights
+		break;
+	case FWidget::EWidgetMode::WM_Scale:
+		Elements = ETransformGizmoSubElements::None; // Attenuation can be scaled using custom gizmo anyways, so scaling doesn't make sense
+		bUseContextCoordinateSystem = false;
+		break;
+	case FWidget::EWidgetMode::WM_2D:
+		Elements = ETransformGizmoSubElements::RotateAxisY | ETransformGizmoSubElements::TranslatePlaneXZ;
+		break;
+	default:
+		Elements = ETransformGizmoSubElements::FullTranslateRotateScale;
+		break;
+	}
+
+	UTransformGizmo* TransformGizmo = GizmoManager->CreateCustomTransformGizmo(Elements);
+	TransformGizmo->bUseContextCoordinateSystem = bUseContextCoordinateSystem;
+	
+	USubTransformProxy* TransformProxy = NewObject<USubTransformProxy>();
+	
+	 // Get the SubTransformProxy from each PointLightGizmo and connect it to the main Transform Proxy
+	for (UInteractiveGizmo* Gizmo : Gizmos)
+	{
+		UPointLightGizmo* LightGizmo = Cast<UPointLightGizmo>(Gizmo);
+
+		if (LightGizmo)
+		{
+			TransformProxy->AddSubTransformProxy(LightGizmo->GetTransformProxy());
+		}
+		
+	}
+
+	TransformGizmo->SetActiveTarget(TransformProxy);
+	TransformGizmo->SetVisibility(SelectedActors.Num() > 0);
+
+	Gizmos.Add(TransformGizmo);
+
+	return Gizmos;
+}
+
+void UPointLightGizmoFactory::ConfigureGridSnapping(bool bGridEnabled, bool bRotGridEnabled, const TArray<UInteractiveGizmo*>& Gizmos) const
+{
+	for (UInteractiveGizmo *Gizmo : Gizmos)
+	{
+		UTransformGizmo* TransformGizmo = Cast<UTransformGizmo>(Gizmo);
+
+		if (TransformGizmo)
+		{
+			TransformGizmo->bSnapToWorldGrid = bGridEnabled;
+			TransformGizmo->bSnapToWorldRotGrid = bRotGridEnabled;
+		}
+	}
+}
diff --git a/Engine/Plugins/Experimental/GizmoEdMode/Source/LightGizmos/Private/ScalableConeGizmo.cpp b/Engine/Plugins/Experimental/GizmoEdMode/Source/LightGizmos/Private/ScalableConeGizmo.cpp
new file mode 100644
--- /dev/null
+++ b/Engine/Plugins/Experimental/GizmoEdMode/Source/LightGizmos/Private/ScalableConeGizmo.cpp
@@ -0,0 +1,318 @@
+// Copyright Epic Games, Inc. All Rights Reserved.
+
+#include "ScalableConeGizmo.h"
+#include "Components/SphereComponent.h"
+#include "BaseGizmos/GizmoBoxComponent.h"
+#include "Engine/CollisionProfile.h"
+#include "BaseGizmos/GizmoMath.h"
+
+// UScalableConeGizmoBuilder
+
+UInteractiveGizmo* UScalableConeGizmoBuilder::BuildGizmo(const FToolBuilderState& SceneState) const
+{
+	UScalableConeGizmo* NewGizmo = NewObject<UScalableConeGizmo>(SceneState.GizmoManager);
+	NewGizmo->SetWorld(SceneState.World);
+	return NewGizmo;
+}
+
+// AScalableConeGizmoActor
+
+AScalableConeGizmoActor::AScalableConeGizmoActor()
+{
+	// root component is a hidden sphere
+	USphereComponent* SphereComponent = CreateDefaultSubobject<USphereComponent>(TEXT("GizmoCenter"));
+	RootComponent = SphereComponent;
+	SphereComponent->InitSphereRadius(1.0f);
+	SphereComponent->SetVisibility(false);
+	SphereComponent->SetCollisionProfileName(UCollisionProfile::NoCollision_ProfileName);
+}
+
+// UScalableConeGizmo
+
+void UScalableConeGizmo::Setup()
+{
+	UInteractiveGizmo::Setup();
+
+	Length = 1000.0f;
+	Angle = 45.f;
+	MaxAngle = 90.f;
+	MinAngle = 0.f;
+	ConeColor = FColor(200, 255, 255);
+
+	UScalableConeGizmoInputBehavior* ScalableConeBehavior = NewObject<UScalableConeGizmoInputBehavior>(this);
+	ScalableConeBehavior->Initialize(this);
+	AddInputBehavior(ScalableConeBehavior);
+
+	CreateGizmoHandles();
+}
+
+void UScalableConeGizmo::Render(IToolsContextRenderAPI* RenderAPI)
+{
+	if (ActiveTarget)
+	{
+		DrawWireSphereCappedCone(RenderAPI->GetPrimitiveDrawInterface(), GizmoActor->GetTransform(), Length, Angle, 32, 8, 10, ConeColor, SDPG_World);
+	}
+}
+
+void UScalableConeGizmo::Shutdown()
+{
+	if (GizmoActor)
+	{
+		GizmoActor->Destroy();
+		GizmoActor = nullptr;
+	}
+}
+
+void UScalableConeGizmo::SetTarget(UTransformProxy* InTarget)
+{
+	ActiveTarget = InTarget;
+
+	// To update the internal GizmoActor when the transform is changed
+	ActiveTarget->OnTransformChanged.AddUObject(this, &UScalableConeGizmo::OnTransformChanged);
+
+	OnTransformChanged(InTarget, InTarget->GetTransform());
+
+}
+
+void UScalableConeGizmo::SetWorld(UWorld* InWorld)
+{
+	World = InWorld;
+}
+
+void UScalableConeGizmo::SetAngleDegrees(float InAngle)
+{
+	Angle = InAngle;
+
+	Angle = FMath::Clamp<float>(Angle, MinAngle, MaxAngle);
+
+	UpdateGizmoHandles();
+
+	if (UpdateAngleFunc)
+	{
+		UpdateAngleFunc(InAngle);
+	}
+}
+
+void UScalableConeGizmo::SetLength(float InLength)
+{
+	Length = InLength;
+}
+
+float UScalableConeGizmo::GetLength()
+{
+	return Length;
+}
+
+float UScalableConeGizmo::GetAngleDegrees()
+{
+	return Angle;
+}
+
+void UScalableConeGizmo::CreateGizmoHandles()
+{
+	FActorSpawnParameters SpawnInfo;
+	GizmoActor = World->SpawnActor<AScalableConeGizmoActor>(FVector::ZeroVector, FRotator::ZeroRotator, SpawnInfo);
+
+	 // Handles are Boxes right now
+	GizmoActor->ScaleHandleYPlus = AGizmoActor::AddDefaultBoxComponent(World, GizmoActor, FLinearColor(1.0, 0, 0), FVector::ZeroVector);
+	GizmoActor->ScaleHandleYMinus = AGizmoActor::AddDefaultBoxComponent(World, GizmoActor, FLinearColor(1.0, 0, 0), FVector::ZeroVector);
+	GizmoActor->ScaleHandleZPlus = AGizmoActor::AddDefaultBoxComponent(World, GizmoActor, FLinearColor(1.0, 0, 0), FVector::ZeroVector);
+	GizmoActor->ScaleHandleZMinus = AGizmoActor::AddDefaultBoxComponent(World, GizmoActor, FLinearColor(1.0, 0, 0), FVector::ZeroVector);
+	
+	UpdateGizmoHandles();
+}
+
+void UScalableConeGizmo::UpdateGizmoHandles()
+{
+	 // Get radius and height of the cone
+	float Radius = Length * FMath::Sin(FMath::DegreesToRadians(Angle));
+	float Height = Length * FMath::Cos(FMath::DegreesToRadians(Angle));
+
+	GizmoActor->ScaleHandleYPlus->SetRelativeLocation(FVector::XAxisVector * Height + FVector::YAxisVector * Radius);
+	GizmoActor->ScaleHandleYMinus->SetRelativeLocation(FVector::XAxisVector * Height - FVector::YAxisVector * Radius);
+	GizmoActor->ScaleHandleZPlus->SetRelativeLocation(FVector::XAxisVector * Height + FVector::ZAxisVector * Radius);
+	GizmoActor->ScaleHandleZMinus->SetRelativeLocation(FVector::XAxisVector * Height - FVector::ZAxisVector * Radius);
+}
+
+void UScalableConeGizmo::OnTransformChanged(UTransformProxy*, FTransform)
+{
+	USceneComponent* GizmoComponent = GizmoActor->GetRootComponent();
+
+	FTransform TargetTransform = ActiveTarget->GetTransform();
+
+	// GizmoActor doesn't want the scale of the object
+	TargetTransform.SetScale3D(FVector(1, 1, 1));
+
+	GizmoComponent->SetWorldTransform(TargetTransform);
+
+	UpdateGizmoHandles();
+}
+
+bool UScalableConeGizmo::HitTest(const FRay& Ray, FHitResult& OutHit, FVector& OutAxis, FTransform& OutTransform)
+{
+	if (!ActiveTarget)
+	{
+		return false;
+	}
+
+	FVector Start = Ray.Origin;
+	const float MaxRaycastDistance = 1e6f;
+	FVector End = Ray.Origin + Ray.Direction * MaxRaycastDistance;
+
+	FCollisionQueryParams Params;
+
+	FQuat Rotation = ActiveTarget->GetTransform().GetRotation();
+
+	// Check which component was hit and update OutAxis and OutTransform
+	if (GizmoActor->ScaleHandleYPlus->LineTraceComponent(OutHit, Start, End, Params))
+	{
+		OutAxis = Rotation.RotateVector(FVector::YAxisVector);
+		OutTransform = GizmoActor->ScaleHandleYPlus->GetComponentTransform();
+		return true;
+	}
+	if (GizmoActor->ScaleHandleYMinus->LineTraceComponent(OutHit, Start, End, Params))
+	{
+		OutAxis = Rotation.RotateVector(-FVector::YAxisVector);
+		OutTransform = GizmoActor->ScaleHandleYMinus->GetComponentTransform();
+		return true;
+	}
+	if (GizmoActor->ScaleHandleZPlus->LineTraceComponent(OutHit, Start, End, Params))
+	{
+		OutAxis = Rotation.RotateVector(FVector::ZAxisVector);
+		OutTransform = GizmoActor->ScaleHandleZPlus->GetComponentTransform();
+		return true;
+	}
+	if (GizmoActor->ScaleHandleZMinus->LineTraceComponent(OutHit, Start, End, Params))
+	{
+		OutAxis = Rotation.RotateVector(-FVector::ZAxisVector);
+		OutTransform = GizmoActor->ScaleHandleZMinus->GetComponentTransform();
+		return true;
+	}
+
+	return false;
+
+}
+
+void UScalableConeGizmo::OnBeginDrag(const FInputDeviceRay& Ray)
+{
+	FVector Start = Ray.WorldRay.Origin;
+	const float MaxRaycastDistance = 1e6f;
+	FVector End = Ray.WorldRay.Origin + Ray.WorldRay.Direction * MaxRaycastDistance;
+
+	FRay HitCheckRay(Start, End - Start);
+	FHitResult HitResult;
+	FTransform DragTransform;
+
+	 // Check if any component was hit
+	if (HitTest(HitCheckRay, HitResult, HitAxis, DragTransform))
+	{
+		FVector RayNearestPt; float RayNearestParam;
+
+		// Get the initial interaction parameters
+		GizmoMath::NearestPointOnLineToRay(DragTransform.GetLocation(), HitAxis,
+			Ray.WorldRay.Origin, Ray.WorldRay.Direction,
+			InteractionStartPoint, InteractionStartParameter,
+			RayNearestPt, RayNearestParam);
+
+		DragStartWorldPosition = DragTransform.GetLocation();
+	}
+}
+
+void UScalableConeGizmo::OnUpdateDrag(const FInputDeviceRay& Ray)
+{
+	FVector AxisNearestPt; float AxisNearestParam;
+	FVector RayNearestPt; float RayNearestParam;
+
+	// Get the current interaction parameters
+	GizmoMath::NearestPointOnLineToRay(DragStartWorldPosition, HitAxis,
+		Ray.WorldRay.Origin, Ray.WorldRay.Direction,
+		AxisNearestPt, AxisNearestParam,
+		RayNearestPt, RayNearestParam);
+
+	FVector GizmoLocation = GizmoActor->GetActorLocation();
+
+	// Vector to the starting position of the interaction
+	FVector Start = InteractionStartPoint - GizmoLocation;
+	Start.Normalize();
+
+	// Vector to the ending position of the interaction
+	FVector End = AxisNearestPt - GizmoLocation;
+	End.Normalize();
+
+	float DotP = FVector::DotProduct(Start, End);
+
+	float DeltaAngle = FMath::Acos(DotP);
+
+	// Get the angle between the start and end vectors and the forward vector to check if the drag direction should be +ve or -ve
+	float StartAngle = FMath::Acos(FVector::DotProduct(Start, GizmoActor->GetActorForwardVector()));
+	float EndAngle = FMath::Acos(FVector::DotProduct(End, GizmoActor->GetActorForwardVector()));
+
+	if (StartAngle > EndAngle)
+	{
+		DeltaAngle = -DeltaAngle;
+	}
+
+	SetAngleDegrees(Angle + FMath::RadiansToDegrees(DeltaAngle));
+
+	InteractionStartPoint = AxisNearestPt;
+	InteractionStartParameter = AxisNearestParam;
+}
+
+// UScalableConeGizmoInputBehavior
+
+void UScalableConeGizmoInputBehavior::Initialize(UScalableConeGizmo* InGizmo)
+{
+	Gizmo = InGizmo;
+}
+
+FInputCaptureRequest UScalableConeGizmoInputBehavior::WantsCapture(const FInputDeviceState& input)
+{
+	if (IsPressed(input))
+	{
+		FHitResult HitResult;
+		FVector HitAxis;
+		FTransform DragTransform;
+		if (Gizmo->HitTest(input.Mouse.WorldRay, HitResult, HitAxis, DragTransform))
+		{
+			return FInputCaptureRequest::Begin(this, EInputCaptureSide::Any, HitResult.Distance);
+		}
+	}
+
+	return FInputCaptureRequest::Ignore();
+}
+
+FInputCaptureUpdate UScalableConeGizmoInputBehavior::BeginCapture(const FInputDeviceState& input, EInputCaptureSide eSide)
+{
+	FInputDeviceRay DeviceRay(input.Mouse.WorldRay, input.Mouse.Position2D);
+	LastWorldRay = DeviceRay.WorldRay;
+	LastScreenPosition = DeviceRay.ScreenPosition;
+
+	Gizmo->OnBeginDrag(DeviceRay);
+	bInputDragCaptured = true;
+	return FInputCaptureUpdate::Begin(this, EInputCaptureSide::Any);
+}
+
+FInputCaptureUpdate UScalableConeGizmoInputBehavior::UpdateCapture(const FInputDeviceState& input, const FInputCaptureData& data)
+{
+	FInputDeviceRay DeviceRay(input.Mouse.WorldRay, input.Mouse.Position2D);
+	LastWorldRay = DeviceRay.WorldRay;
+	LastScreenPosition = DeviceRay.ScreenPosition;
+
+	if (IsReleased(input))
+	{
+		bInputDragCaptured = false;
+		return FInputCaptureUpdate::End();
+	}
+
+	Gizmo->OnUpdateDrag(LastWorldRay);
+
+	return FInputCaptureUpdate::Continue();
+}
+
+void UScalableConeGizmoInputBehavior::ForceEndCapture(const FInputCaptureData& data)
+{
+	if (bInputDragCaptured)
+	{
+		bInputDragCaptured = false;
+	}
+}
\ No newline at end of file
diff --git a/Engine/Plugins/Experimental/GizmoEdMode/Source/LightGizmos/Private/SpotLightGizmo.cpp b/Engine/Plugins/Experimental/GizmoEdMode/Source/LightGizmos/Private/SpotLightGizmo.cpp
new file mode 100644
--- /dev/null
+++ b/Engine/Plugins/Experimental/GizmoEdMode/Source/LightGizmos/Private/SpotLightGizmo.cpp
@@ -0,0 +1,342 @@
+// Copyright Epic Games, Inc. All Rights Reserved.
+
+#include "SpotLightGizmo.h"
+#include "Components/SpotLightComponent.h"
+#include "ScalableConeGizmo.h"
+#include "LightGizmosModule.h"
+#include "Engine/CollisionProfile.h"
+#include "Kismet/KismetMathLibrary.h"
+#include "BaseGizmos/GizmoMath.h"
+#include "Components/SphereComponent.h"
+#include "BaseGizmos/GizmoBoxComponent.h"
+
+// USpotLightGizmoBuilder
+
+UInteractiveGizmo* USpotLightGizmoBuilder::BuildGizmo(const FToolBuilderState& SceneState) const
+{
+	USpotLightGizmo* NewGizmo = NewObject<USpotLightGizmo>(SceneState.GizmoManager);
+	NewGizmo->SetWorld(SceneState.World);
+	return NewGizmo;
+}
+
+// ASpotLightGizmoActor
+
+ASpotLightGizmoActor::ASpotLightGizmoActor()
+{
+	// root component is a hidden sphere
+	USphereComponent* SphereComponent = CreateDefaultSubobject<USphereComponent>(TEXT("GizmoCenter"));
+	RootComponent = SphereComponent;
+	SphereComponent->InitSphereRadius(1.0f);
+	SphereComponent->SetVisibility(false);
+	SphereComponent->SetCollisionProfileName(UCollisionProfile::NoCollision_ProfileName);
+
+}
+
+// USpotLightGizmo
+
+USpotLightGizmo::USpotLightGizmo()
+{
+	LightActor = nullptr;
+	TransformProxy = nullptr;
+}
+
+void USpotLightGizmo::Setup()
+{
+	USpotLightGizmoInputBehavior* SpotLightBehavior = NewObject<USpotLightGizmoInputBehavior>(this);
+	SpotLightBehavior->Initialize(this);
+	AddInputBehavior(SpotLightBehavior);
+}
+
+void USpotLightGizmo::Tick(float DeltaTime)
+{
+	// Make sure the gizmos are up to date with the various light properties 
+	if (OuterAngleGizmo)
+	{
+		OuterAngleGizmo->SetAngleDegrees(LightActor->SpotLightComponent->OuterConeAngle);
+		OuterAngleGizmo->SetLength(LightActor->SpotLightComponent->AttenuationRadius);
+	}
+
+	if (InnerAngleGizmo)
+	{
+		InnerAngleGizmo->SetAngleDegrees(LightActor->SpotLightComponent->InnerConeAngle);
+		InnerAngleGizmo->SetLength(LightActor->SpotLightComponent->AttenuationRadius);
+	}
+
+	if (GizmoActor)
+	{
+		GizmoActor->AttenuationScaleHandle->SetRelativeLocation(FVector(LightActor->SpotLightComponent->AttenuationRadius, 0, 0));
+	}
+}
+
+void USpotLightGizmo::Shutdown()
+{
+	if (OuterAngleGizmo)
+	{
+		GetGizmoManager()->DestroyGizmo(OuterAngleGizmo);
+		OuterAngleGizmo = nullptr;
+	}
+
+	if (InnerAngleGizmo)
+	{
+		GetGizmoManager()->DestroyGizmo(InnerAngleGizmo);
+		InnerAngleGizmo = nullptr;
+	}
+
+	if (GizmoActor)
+	{
+		GizmoActor->Destroy();
+		GizmoActor = nullptr;
+	}
+}
+
+void USpotLightGizmo::SetSelectedObject(ASpotLight* InLight)
+{
+	LightActor = InLight;
+
+	// TODO: No way to remove a component from the transform proxy
+	if (!TransformProxy)
+	{
+		TransformProxy = NewObject<USubTransformProxy>(this);
+	}
+
+	USceneComponent* SceneComponent = LightActor->GetRootComponent();
+	TransformProxy->AddComponent(SceneComponent);
+}
+
+USubTransformProxy* USpotLightGizmo::GetTransformProxy()
+{
+	return TransformProxy;
+}
+
+void USpotLightGizmo::OnOuterAngleUpdate(float NewAngle)
+{
+	LightActor->SpotLightComponent->OuterConeAngle = NewAngle;
+	
+	// OuterAngle cannot be less than Inner Angle
+	if (NewAngle < LightActor->SpotLightComponent->InnerConeAngle)
+	{
+		LightActor->SpotLightComponent->InnerConeAngle = NewAngle;
+	}
+
+	LightActor->SpotLightComponent->MarkRenderStateDirty();
+
+}
+
+void USpotLightGizmo::OnInnerAngleUpdate(float NewAngle)
+{
+	LightActor->SpotLightComponent->InnerConeAngle = NewAngle;
+
+	// Inner Angle cannot be greater than Outer Angle
+	if (NewAngle > LightActor->SpotLightComponent->OuterConeAngle)
+	{
+		LightActor->SpotLightComponent->OuterConeAngle = NewAngle;
+	}
+
+	LightActor->SpotLightComponent->MarkRenderStateDirty();
+}
+
+void USpotLightGizmo::OnTransformChanged(UTransformProxy*, FTransform)
+{
+	USceneComponent* GizmoComponent = GizmoActor->GetRootComponent();
+
+	FTransform TargetTransform = TransformProxy->GetTransform();
+
+	// The gizmo doesn't want the scale of the target
+	TargetTransform.SetScale3D(FVector(1, 1, 1));
+
+	GizmoComponent->SetWorldTransform(TargetTransform);
+}
+
+void USpotLightGizmo::CreateOuterAngleGizmo()
+{
+	if (!LightActor)
+	{
+		return;
+	}
+
+	OuterAngleGizmo = Cast<UScalableConeGizmo>(GetGizmoManager()->CreateGizmo(FLightGizmosModule::ScalableConeGizmoType));
+	OuterAngleGizmo->SetTarget(TransformProxy);
+	OuterAngleGizmo->UpdateAngleFunc = [this](float NewAngle) { this->OnOuterAngleUpdate(NewAngle); };
+	OuterAngleGizmo->MaxAngle = 80.f;
+	OuterAngleGizmo->MinAngle = 1.f;
+}
+
+void USpotLightGizmo::CreateInnerAngleGizmo()
+{
+	if (!LightActor)
+	{
+		return;
+	}
+
+	InnerAngleGizmo = Cast<UScalableConeGizmo>(GetGizmoManager()->CreateGizmo(FLightGizmosModule::ScalableConeGizmoType));
+	InnerAngleGizmo->SetTarget(TransformProxy);
+	InnerAngleGizmo->UpdateAngleFunc = [this](float NewAngle) { this->OnInnerAngleUpdate(NewAngle); };
+	InnerAngleGizmo->MaxAngle = 80.f;
+	InnerAngleGizmo->MinAngle = 1.f;
+	InnerAngleGizmo->ConeColor = FColor(150, 200, 255);
+}
+
+void USpotLightGizmo::SetWorld(UWorld* InWorld)
+{
+	World = InWorld;
+}
+
+void USpotLightGizmo::OnBeginDrag(const FInputDeviceRay& Ray)
+{
+	FVector Start = Ray.WorldRay.Origin;
+	const float MaxRaycastDistance = 1e6f;
+	FVector End = Ray.WorldRay.Origin + Ray.WorldRay.Direction * MaxRaycastDistance;
+
+	FRay HitCheckRay(Start, End - Start);
+	FHitResult HitResult;
+	FVector HitAxis;
+	FTransform DragTransform;
+
+	// Check if any of the components were hit
+	if (HitTest(HitCheckRay, HitResult, DragTransform))
+	{
+		FVector RayNearestPt; float RayNearestParam;
+		FVector InteractionStartPoint;
+
+		// Update interaction start parameters
+		GizmoMath::NearestPointOnLineToRay(DragTransform.GetLocation(), GizmoActor->GetActorForwardVector(),
+			Ray.WorldRay.Origin, Ray.WorldRay.Direction,
+			InteractionStartPoint, InteractionStartParameter,
+			RayNearestPt, RayNearestParam);
+
+		DragStartWorldPosition = DragTransform.GetLocation();
+	}
+}
+
+void USpotLightGizmo::OnUpdateDrag(const FInputDeviceRay& Ray)
+{
+	FVector AxisNearestPt; float AxisNearestParam;
+	FVector RayNearestPt; float RayNearestParam;
+
+	// Get current interaction parameters
+	GizmoMath::NearestPointOnLineToRay(DragStartWorldPosition, GizmoActor->GetActorForwardVector(),
+		Ray.WorldRay.Origin, Ray.WorldRay.Direction,
+		AxisNearestPt, AxisNearestParam,
+		RayNearestPt, RayNearestParam);
+
+	float InteractionCurParameter = AxisNearestParam;
+
+	float DeltaParam = InteractionCurParameter - InteractionStartParameter;
+
+	InteractionStartParameter = InteractionCurParameter;
+
+	// Update the attenuation of the cone
+	float NewAttenuation = LightActor->SpotLightComponent->AttenuationRadius + DeltaParam;
+
+	LightActor->SpotLightComponent->AttenuationRadius = NewAttenuation;
+	LightActor->SpotLightComponent->MarkRenderStateDirty();
+
+	if (OuterAngleGizmo)
+	{
+		OuterAngleGizmo->SetLength(NewAttenuation);
+	}
+
+	if (InnerAngleGizmo)
+	{
+		InnerAngleGizmo->SetLength(NewAttenuation);
+	}
+}
+
+bool USpotLightGizmo::HitTest(const FRay& Ray, FHitResult& OutHit, FTransform& OutTransform)
+{
+	if (!GizmoActor)
+	{
+		return false;
+	}
+
+	FVector Start = Ray.Origin;
+	const float MaxRaycastDistance = 1e6f;
+	FVector End = Ray.Origin + Ray.Direction * MaxRaycastDistance;
+
+	FCollisionQueryParams Params;
+
+	if (GizmoActor->AttenuationScaleHandle->LineTraceComponent(OutHit, Start, End, Params))
+	{
+		OutTransform = GizmoActor->AttenuationScaleHandle->GetComponentTransform();
+		return true;
+	}
+
+	return false;
+}
+
+void USpotLightGizmo::CreateAttenuationScaleGizmo()
+{
+	if (!OuterAngleGizmo)
+	{
+		return;
+	}
+
+	FActorSpawnParameters SpawnInfo;
+	GizmoActor = World->SpawnActor<ASpotLightGizmoActor>(FVector::ZeroVector, FRotator::ZeroRotator, SpawnInfo);
+
+	// The handle to scale attenuation is a box 
+	GizmoActor->AttenuationScaleHandle = AGizmoActor::AddDefaultBoxComponent(World, GizmoActor, FLinearColor::Red, FVector::ZeroVector);
+	GizmoActor->AttenuationScaleHandle->SetRelativeLocation(FVector(LightActor->SpotLightComponent->AttenuationRadius, 0, 0));
+
+	TransformProxy->OnTransformChanged.AddUObject(this, &USpotLightGizmo::OnTransformChanged);
+
+	OnTransformChanged(TransformProxy, TransformProxy->GetTransform());
+}
+
+// USpotLightGizmoInputBehavior
+
+void USpotLightGizmoInputBehavior::Initialize(USpotLightGizmo* InGizmo)
+{
+	Gizmo = InGizmo;
+}
+
+FInputCaptureRequest USpotLightGizmoInputBehavior::WantsCapture(const FInputDeviceState& input)
+{
+	if (IsPressed(input))
+	{
+		FHitResult HitResult;
+		FTransform DragTransform;
+		if (Gizmo->HitTest(input.Mouse.WorldRay, HitResult, DragTransform))
+		{
+			return FInputCaptureRequest::Begin(this, EInputCaptureSide::Any, HitResult.Distance);
+		}
+	}
+
+	return FInputCaptureRequest::Ignore();
+}
+
+FInputCaptureUpdate USpotLightGizmoInputBehavior::BeginCapture(const FInputDeviceState& input, EInputCaptureSide eSide)
+{
+	FInputDeviceRay DeviceRay(input.Mouse.WorldRay, input.Mouse.Position2D);
+	LastWorldRay = DeviceRay.WorldRay;
+	LastScreenPosition = DeviceRay.ScreenPosition;
+
+	Gizmo->OnBeginDrag(DeviceRay);
+	bInputDragCaptured = true;
+	return FInputCaptureUpdate::Begin(this, EInputCaptureSide::Any);
+}
+
+FInputCaptureUpdate USpotLightGizmoInputBehavior::UpdateCapture(const FInputDeviceState& input, const FInputCaptureData& data)
+{
+	FInputDeviceRay DeviceRay(input.Mouse.WorldRay, input.Mouse.Position2D);
+	LastWorldRay = DeviceRay.WorldRay;
+	LastScreenPosition = DeviceRay.ScreenPosition;
+
+	if (IsReleased(input))
+	{
+		bInputDragCaptured = false;
+		return FInputCaptureUpdate::End();
+	}
+
+	Gizmo->OnUpdateDrag(LastWorldRay);
+
+	return FInputCaptureUpdate::Continue();
+}
+
+void USpotLightGizmoInputBehavior::ForceEndCapture(const FInputCaptureData& data)
+{
+	if (bInputDragCaptured)
+	{
+		bInputDragCaptured = false;
+	}
+}
\ No newline at end of file
diff --git a/Engine/Plugins/Experimental/GizmoEdMode/Source/LightGizmos/Private/SpotLightGizmo.h b/Engine/Plugins/Experimental/GizmoEdMode/Source/LightGizmos/Private/SpotLightGizmo.h
new file mode 100644
--- /dev/null
+++ b/Engine/Plugins/Experimental/GizmoEdMode/Source/LightGizmos/Private/SpotLightGizmo.h
@@ -0,0 +1,141 @@
+// Copyright Epic Games, Inc. All Rights Reserved.
+
+#pragma once
+
+#include "ScalableConeGizmo.h"
+#include "Engine/SpotLight.h"
+#include "InteractiveGizmo.h"
+#include "SubTransformProxy.h"
+
+#include "SpotLightGizmo.generated.h"
+
+UCLASS()
+class USpotLightGizmoBuilder : public UInteractiveGizmoBuilder
+{
+	GENERATED_BODY()
+
+public:
+	virtual UInteractiveGizmo* BuildGizmo(const FToolBuilderState& SceneState) const override;
+};
+
+UCLASS()
+class ASpotLightGizmoActor : public AGizmoActor
+{
+	GENERATED_BODY()
+
+public:
+
+	ASpotLightGizmoActor();
+
+	 // The handle to drag and scale the attenuation
+	UPrimitiveComponent* AttenuationScaleHandle;
+};
+
+
+/**
+ * USpotLightGizmo provides a gizmo to allow for editing spot light properties in viewport
+ * Currently supports changing the inner and outer cone angle and scaling the attenuation radius
+ *
+ */
+UCLASS()
+class USpotLightGizmo : public UInteractiveGizmo
+{
+	GENERATED_BODY()
+
+public:
+	// UInteractiveGizmo interface
+
+	virtual void Setup() override;
+
+	virtual void Tick(float DeltaTime) override;
+
+	virtual void Shutdown() override;
+
+	USpotLightGizmo();
+
+	void SetSelectedObject(ASpotLight* InLight);
+
+	/** Create a cone gizmo to change the outer angle of the spotlight */
+	void CreateOuterAngleGizmo();
+
+	/** Create a cone gizmo to change the inner angle of the spotlight */
+	void CreateInnerAngleGizmo();
+
+	/** Create a  gizmo to change the attenuation */
+	void CreateAttenuationScaleGizmo();
+
+	void SetWorld(UWorld* InWorld);
+
+	USubTransformProxy* GetTransformProxy();
+
+	virtual void OnBeginDrag(const FInputDeviceRay& Ray);
+	virtual void OnUpdateDrag(const FInputDeviceRay& Ray);
+
+	/** Check if the Input ray hit any of the components of the gizmo*/
+	bool HitTest(const FRay& Ray, FHitResult& OutHit, FTransform& OutTransform);
+
+private:
+
+	/** The current target light the gizmo is attached to*/
+	UPROPERTY()
+	ASpotLight* LightActor;
+
+	UPROPERTY()
+	UWorld* World;
+
+	/** A transform proxy to use with other gizmos*/
+	UPROPERTY()
+	USubTransformProxy* TransformProxy;
+
+	/** The gizmo to change the outer angle of the spotlight */
+	UPROPERTY()
+	UScalableConeGizmo* OuterAngleGizmo;
+
+	/** The gizmo to change the inner angle of the spotlight */
+	UPROPERTY()
+	UScalableConeGizmo* InnerAngleGizmo;
+
+	/** The internal gizmo actor that is used by the gizmo 
+	 *  We need a GizmoActor separate from the 2 UScalableConeGizmo's
+	 *  so that we can have one handle that scales the attenuation for
+	 *  both of them simultaneously
+	 */
+	UPROPERTY()
+	ASpotLightGizmoActor* GizmoActor;
+
+	UPROPERTY()
+	FVector DragStartWorldPosition;
+
+	UPROPERTY()
+	float InteractionStartParameter;
+
+	void OnOuterAngleUpdate(float NewAngle);
+	void OnInnerAngleUpdate(float NewAngle);
+	void OnTransformChanged(UTransformProxy*, FTransform);
+};
+
+/**
+ * A behavior that forwards clicking and dragging to the gizmo.
+ */
+UCLASS()
+class USpotLightGizmoInputBehavior : public UAnyButtonInputBehavior
+{
+	GENERATED_BODY()
+
+public:
+	virtual FInputCapturePriority GetPriority() override { return FInputCapturePriority(FInputCapturePriority::DEFAULT_GIZMO_PRIORITY); }
+
+	virtual void Initialize(USpotLightGizmo* Gizmo);
+
+	virtual FInputCaptureRequest WantsCapture(const FInputDeviceState& input) override;
+	virtual FInputCaptureUpdate BeginCapture(const FInputDeviceState& input, EInputCaptureSide eSide) override;
+	virtual FInputCaptureUpdate UpdateCapture(const FInputDeviceState& input, const FInputCaptureData& data) override;
+	virtual void ForceEndCapture(const FInputCaptureData& data) override;
+
+protected:
+	USpotLightGizmo* Gizmo;
+	FRay LastWorldRay;
+	FVector2D LastScreenPosition;
+	bool bInputDragCaptured;
+
+};
\ No newline at end of file
diff --git a/Engine/Plugins/Experimental/GizmoEdMode/Source/LightGizmos/Private/SpotLightGizmoFactory.cpp b/Engine/Plugins/Experimental/GizmoEdMode/Source/LightGizmos/Private/SpotLightGizmoFactory.cpp
new file mode 100644
--- /dev/null
+++ b/Engine/Plugins/Experimental/GizmoEdMode/Source/LightGizmos/Private/SpotLightGizmoFactory.cpp
@@ -0,0 +1,128 @@
+// Copyright Epic Games, Inc. All Rights Reserved.
+
+#include "SpotLightGizmoFactory.h"
+#include "EditorModeManager.h"
+#include "Engine/Selection.h"
+#include "BaseGizmos/TransformGizmo.h"
+#include "SubTransformProxy.h"
+#include "Math/Rotator.h"
+#include "UnrealWidget.h"
+#include "Engine/SpotLight.h"
+#include "SpotLightGizmo.h"
+#include "LightGizmosModule.h"
+
+bool USpotLightGizmoFactory::CanBuildGizmoForSelection(FEditorModeTools* ModeTools) const
+{
+	/** Since there can only be one factory active at a time in the current iteration
+	 *  of the tools framework, this function only returns true if /all/ items in the
+	 *  current selection are point lights and then creates point light gizmos for them
+	 *  If not, the selection falls through to the next factory
+	 */
+	TArray<AActor*> SelectedActors;
+	ModeTools->GetSelectedActors()->GetSelectedObjects<AActor>(SelectedActors);
+
+	for (const AActor* Actor : SelectedActors)
+	{
+		if (!Cast<ASpotLight>(Actor))
+		{
+			return false;
+		}
+	}
+
+	// Only return true if all actors in selection are SpotLights
+	return true;
+}
+
+TArray<UInteractiveGizmo*> USpotLightGizmoFactory::BuildGizmoForSelection(FEditorModeTools* ModeTools, UInteractiveGizmoManager* GizmoManager) const
+{
+	TArray<UInteractiveGizmo*> Gizmos;
+	TArray<AActor*> SelectedActors;
+	ModeTools->GetSelectedActors()->GetSelectedObjects<AActor>(SelectedActors);
+
+	bool bCreateInnerConeGizmo = true;
+
+	 // Don't create an inner cone if too many actors are selected to avoid viewport clutter
+	if (SelectedActors.Num() > 5)
+	{
+		bCreateInnerConeGizmo = false;
+	}
+
+	// Create a USpotLightGizmo for each actor in the selection
+	for (AActor* Actor : SelectedActors)
+	{
+		USpotLightGizmo* LightGizmo = Cast<USpotLightGizmo>(GizmoManager->CreateGizmo(FLightGizmosModule::SpotLightGizmoType));
+		LightGizmo->SetSelectedObject(Cast<ASpotLight>(Actor));
+		LightGizmo->CreateOuterAngleGizmo();
+
+		if (bCreateInnerConeGizmo)
+		{
+			LightGizmo->CreateInnerAngleGizmo();
+			LightGizmo->CreateAttenuationScaleGizmo();
+		}
+		Gizmos.Add(LightGizmo);
+	}
+
+	// Create one transform gizmo for the whole selection
+	ETransformGizmoSubElements Elements = ETransformGizmoSubElements::None;
+	bool bUseContextCoordinateSystem = true;
+	FWidget::EWidgetMode WidgetMode = ModeTools->GetWidgetMode();
+
+	switch (WidgetMode)
+	{
+	case FWidget::EWidgetMode::WM_Translate:
+		Elements = ETransformGizmoSubElements::TranslateAllAxes | ETransformGizmoSubElements::TranslateAllPlanes;
+		break;
+	case FWidget::EWidgetMode::WM_Rotate:
+		Elements = ETransformGizmoSubElements::RotateAllAxes;
+		break;
+	case FWidget::EWidgetMode::WM_Scale:
+		Elements = ETransformGizmoSubElements::None; // Attenuation can be scaled using custom gizmo anyways, so scaling doesn't make sense
+		bUseContextCoordinateSystem = false;
+		break;
+	case FWidget::EWidgetMode::WM_2D:
+		Elements = ETransformGizmoSubElements::RotateAxisY | ETransformGizmoSubElements::TranslatePlaneXZ;
+		break;
+	default:
+		Elements = ETransformGizmoSubElements::FullTranslateRotateScale;
+		break;
+	}
+	UTransformGizmo* TransformGizmo = GizmoManager->CreateCustomTransformGizmo(Elements);
+	TransformGizmo->bUseContextCoordinateSystem = bUseContextCoordinateSystem;
+
+	USubTransformProxy* TransformProxy = NewObject<USubTransformProxy>();
+
+	// Get the SubTransformProxy from each UDirectionalLightGizmo and connect it to the main Transform Proxy
+	for (UInteractiveGizmo* Gizmo : Gizmos)
+	{
+		USpotLightGizmo* LightGizmo = Cast<USpotLightGizmo>(Gizmo);
+
+		if (LightGizmo)
+		{
+			TransformProxy->AddSubTransformProxy(LightGizmo->GetTransformProxy());
+		}
+
+	}
+
+	TransformGizmo->SetActiveTarget(TransformProxy);
+	TransformGizmo->SetVisibility(SelectedActors.Num() > 0);
+
+	Gizmos.Add(TransformGizmo);
+
+	return Gizmos;
+}
+
+void USpotLightGizmoFactory::ConfigureGridSnapping(bool bGridEnabled, bool bRotGridEnabled, const TArray<UInteractiveGizmo*>& Gizmos) const
+{
+	for (UInteractiveGizmo* Gizmo : Gizmos)
+	{
+		UTransformGizmo* TransformGizmo = Cast<UTransformGizmo>(Gizmo);
+
+		if (TransformGizmo)
+		{
+			TransformGizmo->bSnapToWorldGrid = bGridEnabled;
+			TransformGizmo->bSnapToWorldRotGrid = bRotGridEnabled;
+		}
+	}
+}
+
+
diff --git a/Engine/Plugins/Experimental/GizmoEdMode/Source/LightGizmos/Private/SubTransformProxy.cpp b/Engine/Plugins/Experimental/GizmoEdMode/Source/LightGizmos/Private/SubTransformProxy.cpp
new file mode 100644
--- /dev/null
+++ b/Engine/Plugins/Experimental/GizmoEdMode/Source/LightGizmos/Private/SubTransformProxy.cpp
@@ -0,0 +1,64 @@
+// Copyright Epic Games, Inc. All Rights Reserved.
+
+#include "SubTransformProxy.h"
+
+USubTransformProxy::USubTransformProxy()
+{
+	RelativeTransform = FTransform::Identity;
+}
+
+FTransform USubTransformProxy::AddSubTransformProxy(USubTransformProxy* InProxy, bool bSubscribeToChanges)
+{
+	check(InProxy);
+
+	FTransform OriginalTransform = InProxy->GetTransform();
+
+	// Add each object to the proxy
+	for (FRelativeObject& Object : InProxy->Objects)
+	{
+		FRelativeObject NewObj;
+		NewObj.Component = Object.Component;
+		NewObj.bModifyComponentOnTransform = Object.bModifyComponentOnTransform;
+		NewObj.StartTransform = Object.Component->GetComponentToWorld();
+		NewObj.RelativeTransform = FTransform::Identity;
+		Objects.Add(NewObj);
+	}
+
+	UpdateSharedTransform();
+
+	// Get the relative transform of the InProxy we just added
+	OriginalTransform.SetToRelativeTransform(SharedTransform);
+
+	if (bSubscribeToChanges)
+	{
+		OnTransformChanged.AddUObject(InProxy, &USubTransformProxy::OnParentTransformChanged);
+		OnRelativeTransformChanged.AddUObject(InProxy, &USubTransformProxy::OnParentRelativeTransformChanged);
+	}
+
+	OnRelativeTransformChanged.Broadcast(this, SharedTransform);
+	OnTransformChanged.Broadcast(this, SharedTransform);
+
+	return OriginalTransform;
+}
+
+void USubTransformProxy::OnParentTransformChanged(UTransformProxy* Parent, FTransform ParentTransform)
+{
+	FTransform OriginalTransform = SharedTransform;
+
+	FTransform::Multiply(&SharedTransform, &RelativeTransform, &ParentTransform);
+
+	 // Parent Transform changing doesn't have to mean that your transform changed
+	if (!SharedTransform.Equals(OriginalTransform))
+	{
+		UpdateObjectTransforms();
+
+		OnTransformChanged.Broadcast(this, SharedTransform);
+	}
+
+}
+
+void USubTransformProxy::OnParentRelativeTransformChanged(UTransformProxy* Parent, FTransform ParentTransform)
+{
+	RelativeTransform = SharedTransform.GetRelativeTransform(ParentTransform);
+}
+
diff --git a/Engine/Plugins/Experimental/GizmoEdMode/Source/LightGizmos/Private/SubTransformProxy.h b/Engine/Plugins/Experimental/GizmoEdMode/Source/LightGizmos/Private/SubTransformProxy.h
new file mode 100644
--- /dev/null
+++ b/Engine/Plugins/Experimental/GizmoEdMode/Source/LightGizmos/Private/SubTransformProxy.h
@@ -0,0 +1,44 @@
+// Copyright Epic Games, Inc. All Rights Reserved.
+
+#pragma once
+
+#include "BaseGizmos/TransformProxy.h"
+
+#include "SubTransformProxy.generated.h"
+
+/*
+ * A sub-class of UTransformProxy that adds support for sub-proxies that can be attached
+ * to a Transform Proxy
+*/
+UCLASS()
+class USubTransformProxy : public UTransformProxy
+{
+	GENERATED_BODY()
+
+	public:
+
+	USubTransformProxy();
+
+	/**
+	 * Attach a sub transform proxy
+	 * @param bSubscribetoSchanges true if the SubProxy wants to listen to parent's updates
+	 * @return The relative transform of the SubProxy that was passed in
+	 */
+	virtual FTransform AddSubTransformProxy(USubTransformProxy* InProxy, bool bSubscribeToChanges = true);
+
+	/**
+	 * This delegate is fired whenever the internal transform changes in a way that would change the relative transform, ie
+	 * on AddSubTransformProxy and Add Component
+	 */
+	DECLARE_MULTICAST_DELEGATE_TwoParams(FOnRelativeTransformChanged, UTransformProxy*, FTransform);
+	FOnRelativeTransformChanged OnRelativeTransformChanged;
+
+	protected:
+
+	virtual void OnParentTransformChanged(UTransformProxy* Parent, FTransform ParentTransform);
+	virtual void OnParentRelativeTransformChanged(UTransformProxy* Parent, FTransform ParentTransform);
+
+	/** The relative transform of this SubTransformProxy if it is attached to a parent SubTransformProxy */
+	UPROPERTY()
+	FTransform RelativeTransform;
+};
\ No newline at end of file
diff --git a/Engine/Plugins/Experimental/GizmoEdMode/Source/LightGizmos/Public/DirectionalLightGizmoFactory.h b/Engine/Plugins/Experimental/GizmoEdMode/Source/LightGizmos/Public/DirectionalLightGizmoFactory.h
new file mode 100644
--- /dev/null
+++ b/Engine/Plugins/Experimental/GizmoEdMode/Source/LightGizmos/Public/DirectionalLightGizmoFactory.h
@@ -0,0 +1,19 @@
+// Copyright Epic Games, Inc. All Rights Reserved.
+
+#pragma once
+
+#include "AssetEditorGizmoFactory.h"
+
+#include "DirectionalLightGizmoFactory.generated.h"
+
+UCLASS()
+class LIGHTGIZMOS_API UDirectionalLightGizmoFactory : public UObject, public IAssetEditorGizmoFactory
+{
+	GENERATED_BODY()
+public:
+	//IAssetEditorGizmoFactory interface
+	virtual bool CanBuildGizmoForSelection(FEditorModeTools* ModeTools) const override;
+	virtual TArray<UInteractiveGizmo*> BuildGizmoForSelection(FEditorModeTools* ModeTools, UInteractiveGizmoManager* GizmoManager) const override;
+	virtual EAssetEditorGizmoFactoryPriority GetPriority() const override { return EAssetEditorGizmoFactoryPriority::Normal; }
+	virtual void ConfigureGridSnapping(bool bGridEnabled, bool bRotGridEnabled, const TArray<UInteractiveGizmo*> &Gizmo) const override;
+};
diff --git a/Engine/Plugins/Experimental/GizmoEdMode/Source/LightGizmos/Public/LightGizmosModule.h b/Engine/Plugins/Experimental/GizmoEdMode/Source/LightGizmos/Public/LightGizmosModule.h
new file mode 100644
--- /dev/null
+++ b/Engine/Plugins/Experimental/GizmoEdMode/Source/LightGizmos/Public/LightGizmosModule.h
@@ -0,0 +1,23 @@
+// Copyright Epic Games, Inc. All Rights Reserved.
+
+#pragma once
+
+#include "CoreMinimal.h"
+#include "Modules/ModuleManager.h"
+
+class FLightGizmosModule : public IModuleInterface
+{
+public:
+
+	/** IModuleInterface implementation */
+	virtual void OnPostEngineInit();
+
+	virtual void StartupModule() override;
+	virtual void ShutdownModule() override;
+
+	/** Identifiers for the different light gizmo types*/
+	static FString PointLightGizmoType;
+	static FString SpotLightGizmoType;
+	static FString ScalableConeGizmoType;
+	static FString DirectionalLightGizmoType;
+};
diff --git a/Engine/Plugins/Experimental/GizmoEdMode/Source/LightGizmos/Public/PointLightGizmoFactory.h b/Engine/Plugins/Experimental/GizmoEdMode/Source/LightGizmos/Public/PointLightGizmoFactory.h
new file mode 100644
--- /dev/null
+++ b/Engine/Plugins/Experimental/GizmoEdMode/Source/LightGizmos/Public/PointLightGizmoFactory.h
@@ -0,0 +1,19 @@
+// Copyright Epic Games, Inc. All Rights Reserved.
+
+#pragma once
+
+#include "AssetEditorGizmoFactory.h"
+
+#include "PointLightGizmoFactory.generated.h"
+
+UCLASS()
+class LIGHTGIZMOS_API UPointLightGizmoFactory : public UObject, public IAssetEditorGizmoFactory
+{
+	GENERATED_BODY()
+public:
+	//IAssetEditorGizmoFactory interface
+	virtual bool CanBuildGizmoForSelection(FEditorModeTools* ModeTools) const override;
+	virtual TArray<UInteractiveGizmo*> BuildGizmoForSelection(FEditorModeTools* ModeTools, UInteractiveGizmoManager* GizmoManager) const override;
+	virtual EAssetEditorGizmoFactoryPriority GetPriority() const override { return EAssetEditorGizmoFactoryPriority::Normal; }
+	virtual void ConfigureGridSnapping(bool bGridEnabled, bool bRotGridEnabled, const TArray<UInteractiveGizmo*>& Gizmo) const override;
+};
diff --git a/Engine/Plugins/Experimental/GizmoEdMode/Source/LightGizmos/Public/ScalableConeGizmo.h b/Engine/Plugins/Experimental/GizmoEdMode/Source/LightGizmos/Public/ScalableConeGizmo.h
new file mode 100644
--- /dev/null
+++ b/Engine/Plugins/Experimental/GizmoEdMode/Source/LightGizmos/Public/ScalableConeGizmo.h
@@ -0,0 +1,161 @@
+// Copyright Epic Games, Inc. All Rights Reserved.
+
+#pragma once
+
+#include "CoreMinimal.h"
+#include "InteractiveGizmo.h"
+#include "InteractiveGizmoBuilder.h"
+#include "BaseGizmos/TransformProxy.h"
+#include "BaseGizmos/GizmoActor.h"
+#include "BaseBehaviors/BehaviorTargetInterfaces.h"
+#include "BaseBehaviors/AnyButtonInputBehavior.h"
+
+#include "ScalableConeGizmo.generated.h"
+
+UCLASS()
+class LIGHTGIZMOS_API UScalableConeGizmoBuilder : public UInteractiveGizmoBuilder
+{
+	GENERATED_BODY()
+
+public:
+	virtual UInteractiveGizmo* BuildGizmo(const FToolBuilderState& SceneState) const override;
+};
+
+UCLASS()
+class LIGHTGIZMOS_API AScalableConeGizmoActor : public AGizmoActor
+{
+	GENERATED_BODY()
+
+public:
+
+	AScalableConeGizmoActor();
+
+	UPrimitiveComponent* ScaleHandleYPlus;
+	UPrimitiveComponent* ScaleHandleYMinus;
+	UPrimitiveComponent* ScaleHandleZPlus;
+	UPrimitiveComponent* ScaleHandleZMinus;
+
+	UPrimitiveComponent* LengthHandle;
+};
+
+/**
+ * UScalableConeGizmo provides a cone that can be scaled (changing its angle)
+ * The in-scene representation of the Gizmo is a AScalableConeGizmoActor (or subclass).
+ */
+UCLASS()
+class LIGHTGIZMOS_API UScalableConeGizmo : public UInteractiveGizmo, public IHoverBehaviorTarget
+{
+	GENERATED_BODY()
+
+public:
+	// UInteractiveGizmo interface
+
+	virtual void Setup() override;
+
+	virtual void Render(IToolsContextRenderAPI* RenderAPI) override;
+
+	virtual void Shutdown() override;
+
+	// Set the target to attach the gizmo to
+	virtual void SetTarget(UTransformProxy* InTarget);
+	virtual void SetWorld(UWorld* InWorld);
+
+	// Gettors and Settors for the Angle and Length
+	void SetAngleDegrees(float InAngle);
+	void SetLength(float InLength);
+	float GetLength();
+	float GetAngleDegrees();
+
+	// IHoverBehaviorTarget interface
+	virtual FInputRayHit BeginHoverSequenceHitTest(const FInputDeviceRay& PressPos) override { return FInputRayHit(); }
+	virtual void OnBeginHover(const FInputDeviceRay& DevicePos) override {}
+	virtual bool OnUpdateHover(const FInputDeviceRay& DevicePos) override { return true; }
+	virtual void OnEndHover() override {}
+
+	virtual void OnBeginDrag(const FInputDeviceRay& Ray);
+	virtual void OnUpdateDrag(const FInputDeviceRay& Ray);
+	bool HitTest(const FRay& Ray, FHitResult& OutHit, FVector& OutAxis, FTransform& OutTransform);
+
+	// The maximum angle the cone can be stretched to
+	UPROPERTY()
+	float MaxAngle;
+
+	// The minimum angle the cone can be stretched to
+	UPROPERTY()
+	float MinAngle;
+
+	// The color of the cone
+	UPROPERTY()
+	FColor ConeColor;
+
+	/** Called when the Angle of the cone is changed. Sends new angle as parameter. */
+	TFunction<void(const float)> UpdateAngleFunc = nullptr;
+
+private:
+
+	void CreateGizmoHandles();
+	void UpdateGizmoHandles();
+	void OnTransformChanged(UTransformProxy*, FTransform);
+
+	// The ConeLength
+	UPROPERTY()
+	float Length;
+
+	UPROPERTY()
+	float Angle;
+
+	UPROPERTY()
+	UTransformProxy* ActiveTarget;
+
+	UPROPERTY()
+	AScalableConeGizmoActor* GizmoActor;
+
+	UPROPERTY()
+	UWorld* World;
+
+	/** Used for calculations when moving the handles*/
+
+	UPROPERTY()
+	FVector DragStartWorldPosition;
+
+	UPROPERTY()
+	FVector InteractionStartPoint;
+
+	UPROPERTY()
+	float InteractionStartParameter;
+
+	UPROPERTY()
+	FVector HitAxis;
+
+	UPROPERTY()
+	FVector RotationPlaneX;
+
+	UPROPERTY()
+	FVector RotationPlaneY;
+};
+
+/**
+ * A behavior that forwards clicking and dragging to the gizmo.
+ */
+UCLASS()
+class LIGHTGIZMOS_API UScalableConeGizmoInputBehavior : public UAnyButtonInputBehavior
+{
+	GENERATED_BODY()
+
+public:
+	virtual FInputCapturePriority GetPriority() override { return FInputCapturePriority(FInputCapturePriority::DEFAULT_GIZMO_PRIORITY); }
+
+	virtual void Initialize(UScalableConeGizmo* Gizmo);
+
+	virtual FInputCaptureRequest WantsCapture(const FInputDeviceState& input) override;
+	virtual FInputCaptureUpdate BeginCapture(const FInputDeviceState& input, EInputCaptureSide eSide) override;
+	virtual FInputCaptureUpdate UpdateCapture(const FInputDeviceState& input, const FInputCaptureData& data) override;
+	virtual void ForceEndCapture(const FInputCaptureData& data) override;
+
+protected:
+	UScalableConeGizmo* Gizmo;
+	FRay LastWorldRay;
+	FVector2D LastScreenPosition;
+	bool bInputDragCaptured;
+
+};
\ No newline at end of file
diff --git a/Engine/Plugins/Experimental/GizmoEdMode/Source/LightGizmos/Public/SpotLightGizmoFactory.h b/Engine/Plugins/Experimental/GizmoEdMode/Source/LightGizmos/Public/SpotLightGizmoFactory.h
new file mode 100644
--- /dev/null
+++ b/Engine/Plugins/Experimental/GizmoEdMode/Source/LightGizmos/Public/SpotLightGizmoFactory.h
@@ -0,0 +1,19 @@
+// Copyright Epic Games, Inc. All Rights Reserved.
+
+#pragma once
+
+#include "AssetEditorGizmoFactory.h"
+
+#include "SpotLightGizmoFactory.generated.h"
+
+UCLASS()
+class LIGHTGIZMOS_API USpotLightGizmoFactory : public UObject, public IAssetEditorGizmoFactory
+{
+	GENERATED_BODY()
+public:
+	//IAssetEditorGizmoFactory interface
+	virtual bool CanBuildGizmoForSelection(FEditorModeTools* ModeTools) const override;
+	virtual TArray<UInteractiveGizmo*> BuildGizmoForSelection(FEditorModeTools* ModeTools, UInteractiveGizmoManager* GizmoManager) const override;
+	virtual EAssetEditorGizmoFactoryPriority GetPriority() const override { return EAssetEditorGizmoFactoryPriority::Normal; }
+	virtual void ConfigureGridSnapping(bool bGridEnabled, bool bRotGridEnabled, const TArray<UInteractiveGizmo*>& Gizmo) const override;
+};
diff --git a/Engine/Source/Runtime/Experimental/InteractiveToolsFramework/Private/BaseGizmos/GizmoLineHandleComponent.cpp b/Engine/Source/Runtime/Experimental/InteractiveToolsFramework/Private/BaseGizmos/GizmoLineHandleComponent.cpp
--- a/Engine/Source/Runtime/Experimental/InteractiveToolsFramework/Private/BaseGizmos/GizmoLineHandleComponent.cpp
+++ b/Engine/Source/Runtime/Experimental/InteractiveToolsFramework/Private/BaseGizmos/GizmoLineHandleComponent.cpp
@@ -68,8 +68,8 @@ public:
 				}
 
 				FVector ScaledDiskOrigin = LengthScale * (WorldDiskOrigin - WorldBaseOrigin) + WorldBaseOrigin;
-				FVector ScaledIntevalStart = -LengthScale * (WorldIntervalEnd - WorldDiskOrigin) + WorldDiskOrigin;
-				FVector ScaledIntevalEnd = LengthScale * (WorldIntervalEnd - WorldDiskOrigin) + WorldDiskOrigin;
+				FVector ScaledIntevalStart = -LengthScale * (WorldIntervalEnd - WorldDiskOrigin) + ScaledDiskOrigin;
+				FVector ScaledIntevalEnd = LengthScale * (WorldIntervalEnd - WorldDiskOrigin) + ScaledDiskOrigin;
 
 				float UseThickness = (bExternalHoverState != nullptr && *bExternalHoverState == true) ?
 					(HoverThicknessMultiplier * Thickness) : (Thickness);
diff --git a/Engine/Source/Runtime/Experimental/InteractiveToolsFramework/Private/BaseGizmos/ScalableSphereGizmo.cpp b/Engine/Source/Runtime/Experimental/InteractiveToolsFramework/Private/BaseGizmos/ScalableSphereGizmo.cpp
new file mode 100644
--- /dev/null
+++ b/Engine/Source/Runtime/Experimental/InteractiveToolsFramework/Private/BaseGizmos/ScalableSphereGizmo.cpp
@@ -0,0 +1,314 @@
+// Copyright Epic Games, Inc. All Rights Reserved.
+
+#include "BaseGizmos/ScalableSphereGizmo.h"
+#include "InteractiveGizmoManager.h"
+#include "SceneManagement.h"
+#include "BaseGizmos/GizmoBoxComponent.h"
+#include "Components/SphereComponent.h"
+#include "Engine/CollisionProfile.h"
+#include "Kismet/KismetMathLibrary.h"
+#include "BaseGizmos/GizmoMath.h"
+
+// UScalableSphereGizmoBuilder
+
+UInteractiveGizmo* UScalableSphereGizmoBuilder::BuildGizmo(const FToolBuilderState& SceneState) const
+{
+	UScalableSphereGizmo* NewGizmo = NewObject<UScalableSphereGizmo>(SceneState.GizmoManager);
+
+	// Have to set world to be able to spawn actor
+	NewGizmo->SetWorld(SceneState.World);
+	
+	return NewGizmo;
+}
+
+// AScalableSphereGizmoActor
+
+AScalableSphereGizmoActor::AScalableSphereGizmoActor()
+{
+	// root component is a hidden sphere
+	USphereComponent* SphereComponent = CreateDefaultSubobject<USphereComponent>(TEXT("GizmoCenter"));
+	RootComponent = SphereComponent;
+	SphereComponent->InitSphereRadius(1.0f);
+	SphereComponent->SetVisibility(false);
+	SphereComponent->SetCollisionProfileName(UCollisionProfile::NoCollision_ProfileName);
+}
+
+// UScalableSphereGizmo
+
+void UScalableSphereGizmo::Setup()
+{
+	UInteractiveGizmo::Setup();
+
+	Radius = 1000.0f;
+
+	UScalableSphereGizmoInputBehavior* ScalableSphereBehavior = NewObject<UScalableSphereGizmoInputBehavior>(this);
+	ScalableSphereBehavior->Initialize(this);
+	AddInputBehavior(ScalableSphereBehavior);
+
+	CreateGizmoHandles();
+}
+
+void UScalableSphereGizmo::Render(IToolsContextRenderAPI* RenderAPI)
+{
+	if (ActiveTarget)
+	{
+		DrawWireSphereAutoSides(RenderAPI->GetPrimitiveDrawInterface(), ActiveTarget->GetTransform(), FColor(200, 255, 255), Radius, SDPG_World);
+	}
+}
+
+void UScalableSphereGizmo::Shutdown()
+{
+	if (GizmoActor)
+	{
+		GizmoActor->Destroy();
+		GizmoActor = nullptr;
+	}
+}
+
+void UScalableSphereGizmo::SetTarget(UTransformProxy* InTarget)
+{
+	ActiveTarget = InTarget;
+
+	// Make sure the internal GizmoActor updates when the transform of the target is changed
+	ActiveTarget->OnTransformChanged.AddUObject(this, &UScalableSphereGizmo::OnTransformChanged);
+	
+	OnTransformChanged(InTarget, InTarget->GetTransform());
+}
+
+void UScalableSphereGizmo::SetWorld(UWorld* InWorld)
+{
+	World = InWorld;
+
+}
+
+void UScalableSphereGizmo::CreateGizmoHandles()
+{
+	FActorSpawnParameters SpawnInfo;
+	GizmoActor = World->SpawnActor<AScalableSphereGizmoActor>(FVector::ZeroVector, FRotator::ZeroRotator, SpawnInfo);
+
+	// Create all the 6 handles in each direction
+
+	auto NewHandle = [this] { return AGizmoActor::AddDefaultBoxComponent(World, GizmoActor, FLinearColor(1.0f, 0.0f, 0.0f), FVector::ZeroVector); };
+	GizmoActor->XPositive = NewHandle();
+	GizmoActor->XNegative = NewHandle();
+	GizmoActor->YPositive = NewHandle();
+	GizmoActor->YNegative = NewHandle();
+	GizmoActor->ZPositive = NewHandle();
+	GizmoActor->ZNegative = NewHandle();
+
+	UpdateGizmoHandles();
+
+}
+
+void UScalableSphereGizmo::UpdateGizmoHandles()
+{
+	GizmoActor->XPositive->SetRelativeLocation(FVector::XAxisVector * Radius);
+	GizmoActor->XNegative->SetRelativeLocation(-FVector::XAxisVector * Radius);
+
+	GizmoActor->YPositive->SetRelativeLocation(FVector::YAxisVector * Radius);
+	GizmoActor->YNegative->SetRelativeLocation(-FVector::YAxisVector * Radius);
+
+	GizmoActor->ZPositive->SetRelativeLocation(FVector::ZAxisVector * Radius);
+	GizmoActor->ZNegative->SetRelativeLocation(-FVector::ZAxisVector * Radius);
+}
+
+void UScalableSphereGizmo::OnTransformChanged(UTransformProxy*, FTransform)
+{
+	USceneComponent* GizmoComponent = GizmoActor->GetRootComponent();
+
+	FTransform TargetTransform = ActiveTarget->GetTransform();
+
+	// Don't scale the internal gizmo actor (only update position and rotation changes)
+	TargetTransform.SetScale3D(FVector(1, 1, 1));
+
+	GizmoComponent->SetWorldTransform(TargetTransform);
+}
+
+bool UScalableSphereGizmo::HitTest(const FRay& Ray, FHitResult& OutHit, FVector& OutAxis, FTransform &OutTransform)
+{
+	if (!ActiveTarget)
+	{
+		return false;
+	}
+
+	FVector Start = Ray.Origin;
+	const float MaxRaycastDistance = 1e6f;
+	FVector End = Ray.Origin + Ray.Direction * MaxRaycastDistance;
+
+	FCollisionQueryParams Params;
+
+	// Check each actor to see if any of them were hit
+	if (GizmoActor->XPositive->LineTraceComponent(OutHit, Start, End, Params))
+	{
+		OutAxis = FVector::XAxisVector;
+		OutTransform = GizmoActor->XPositive->GetComponentTransform();
+		return true;
+	}
+	if (GizmoActor->XNegative->LineTraceComponent(OutHit, Start, End, Params))
+	{
+		OutAxis = -FVector::XAxisVector;
+		OutTransform = GizmoActor->XNegative->GetComponentTransform();
+		return true;
+	}
+
+	if (GizmoActor->YPositive->LineTraceComponent(OutHit, Start, End, Params))
+	{
+		OutAxis = FVector::YAxisVector;
+		OutTransform = GizmoActor->YPositive->GetComponentTransform();
+		return true;
+	}
+	if (GizmoActor->YNegative->LineTraceComponent(OutHit, Start, End, Params))
+	{
+		OutAxis = -FVector::YAxisVector;
+		OutTransform = GizmoActor->YNegative->GetComponentTransform();
+		return true;
+	}
+
+	if (GizmoActor->ZPositive->LineTraceComponent(OutHit, Start, End, Params))
+	{
+		OutAxis = FVector::ZAxisVector;
+		OutTransform = GizmoActor->ZPositive->GetComponentTransform();
+		return true;
+	}
+	if (GizmoActor->ZNegative->LineTraceComponent(OutHit, Start, End, Params))
+	{
+		OutAxis = -FVector::ZAxisVector;
+		OutTransform = GizmoActor->ZNegative->GetComponentTransform();
+		return true;
+	}
+
+	return false;
+}
+
+void UScalableSphereGizmo::SetRadius(float InRadius)
+{
+	// Negative Radius not allowed
+	if (InRadius < 0)
+	{
+		InRadius = 0;
+	}
+
+	Radius = InRadius;
+
+	if (UpdateRadiusFunc)
+	{
+		UpdateRadiusFunc(Radius);
+	}
+
+	UpdateGizmoHandles();
+}
+
+void UScalableSphereGizmo::OnBeginDrag(const FInputDeviceRay& Ray)
+{
+	FVector Start = Ray.WorldRay.Origin;
+	const float MaxRaycastDistance = 1e6f;
+	FVector End = Ray.WorldRay.Origin + Ray.WorldRay.Direction * MaxRaycastDistance;
+
+	FRay HitCheckRay(Start, End - Start);
+	FHitResult HitResult;
+	FVector HitAxis;
+	FTransform DragTransform;
+
+	 // Check if the Ray hit any of the handles
+	if (HitTest(HitCheckRay, HitResult, HitAxis, DragTransform))
+	{
+		ActiveAxis = HitAxis;
+
+		FVector RayNearestPt; 
+		float RayNearestParam;
+		FVector InteractionStartPoint;
+
+		 // Find the initial parameters along the hit axis
+		GizmoMath::NearestPointOnLineToRay(DragTransform.GetLocation(), ActiveAxis,
+			Ray.WorldRay.Origin, Ray.WorldRay.Direction,
+			InteractionStartPoint, InteractionStartParameter,
+			RayNearestPt, RayNearestParam);
+
+		DragStartWorldPosition = DragTransform.GetLocation();
+	}
+}
+
+void UScalableSphereGizmo::OnUpdateDrag(const FInputDeviceRay& Ray)
+{
+	FVector AxisNearestPt; 
+	float AxisNearestParam;
+	FVector RayNearestPt; 
+	float RayNearestParam;
+
+	 // Find the parameters along the hit axis
+	GizmoMath::NearestPointOnLineToRay(DragStartWorldPosition, ActiveAxis,
+		Ray.WorldRay.Origin, Ray.WorldRay.Direction,
+		AxisNearestPt, AxisNearestParam,
+		RayNearestPt, RayNearestParam);
+
+	float InteractionCurrentParameter = AxisNearestParam;
+
+	float DeltaParam = InteractionCurrentParameter - InteractionStartParameter;
+
+	InteractionStartParameter = InteractionCurrentParameter;
+
+	 // Update the radius
+	SetRadius(Radius + DeltaParam);
+}
+
+// UScalableSphereGizmoInputBehavior
+
+void UScalableSphereGizmoInputBehavior::Initialize(UScalableSphereGizmo* InGizmo)
+{
+	Gizmo = InGizmo;
+}
+
+FInputCaptureRequest UScalableSphereGizmoInputBehavior::WantsCapture(const FInputDeviceState& input)
+{
+	if (IsPressed(input))
+	{
+		FHitResult HitResult;
+		FVector HitAxis;
+		FTransform DragTransform;
+		if (Gizmo->HitTest(input.Mouse.WorldRay, HitResult, HitAxis, DragTransform))
+		{
+			return FInputCaptureRequest::Begin(this, EInputCaptureSide::Any, HitResult.Distance);
+		}
+	}
+
+	return FInputCaptureRequest::Ignore();
+}
+
+FInputCaptureUpdate UScalableSphereGizmoInputBehavior::BeginCapture(const FInputDeviceState& input, EInputCaptureSide eSide)
+{
+	FInputDeviceRay DeviceRay(input.Mouse.WorldRay, input.Mouse.Position2D);
+	LastWorldRay = DeviceRay.WorldRay;
+	LastScreenPosition = DeviceRay.ScreenPosition;
+
+	// Forward behavior to the Gizmo
+	Gizmo->OnBeginDrag(DeviceRay);
+
+	bInputDragCaptured = true;
+	return FInputCaptureUpdate::Begin(this, EInputCaptureSide::Any);
+}
+
+FInputCaptureUpdate UScalableSphereGizmoInputBehavior::UpdateCapture(const FInputDeviceState& input, const FInputCaptureData& data)
+{
+	FInputDeviceRay DeviceRay(input.Mouse.WorldRay, input.Mouse.Position2D);
+	LastWorldRay = DeviceRay.WorldRay;
+	LastScreenPosition = DeviceRay.ScreenPosition;
+
+	if (IsReleased(input))
+	{
+		bInputDragCaptured = false;
+		return FInputCaptureUpdate::End();
+	}
+
+	// Forward behavior to the Gizmo
+	Gizmo->OnUpdateDrag(LastWorldRay);
+
+	return FInputCaptureUpdate::Continue();
+}
+
+void UScalableSphereGizmoInputBehavior::ForceEndCapture(const FInputCaptureData& data)
+{
+	if (bInputDragCaptured)
+	{
+		bInputDragCaptured = false;
+	}
+}
\ No newline at end of file
diff --git a/Engine/Source/Runtime/Experimental/InteractiveToolsFramework/Private/InteractiveGizmoManager.cpp b/Engine/Source/Runtime/Experimental/InteractiveToolsFramework/Private/InteractiveGizmoManager.cpp
--- a/Engine/Source/Runtime/Experimental/InteractiveToolsFramework/Private/InteractiveGizmoManager.cpp
+++ b/Engine/Source/Runtime/Experimental/InteractiveToolsFramework/Private/InteractiveGizmoManager.cpp
@@ -9,6 +9,7 @@
 #include "BaseGizmos/AxisAngleGizmo.h"
 #include "BaseGizmos/TransformGizmo.h"
 #include "BaseGizmos/IntervalGizmo.h"
+#include "BaseGizmos/ScalableSphereGizmo.h"
 
 #define LOCTEXT_NAMESPACE "UInteractiveGizmoManager"
 
@@ -48,6 +49,7 @@ void UInteractiveGizmoManager::Shutdown()
 		DeregisterGizmoType(DefaultPlanePositionBuilderIdentifier);
 		DeregisterGizmoType(DefaultAxisAngleBuilderIdentifier);
 		DeregisterGizmoType(DefaultThreeAxisTransformBuilderIdentifier);
+		DeregisterGizmoType(DefaultScalableSphereBuilderIdentifier);
 	}
 }
 
@@ -129,19 +131,8 @@ UInteractiveGizmo* UInteractiveGizmoManager::CreateGizmo(const FString& BuilderI
 
 bool UInteractiveGizmoManager::DestroyGizmo(UInteractiveGizmo* Gizmo)
 {
-	int FoundIndex = -1;
-	for ( int i = 0; i < ActiveGizmos.Num(); ++i )
-	{
-		if (ActiveGizmos[i].Gizmo == Gizmo)
-		{
-			FoundIndex = i;
-			break;
-		}
-	}
-	if (FoundIndex == -1)
-	{
-		return false;
-	}
+	auto Pred = [Gizmo](const FActiveGizmo& ActiveGizmo) {return ActiveGizmo.Gizmo == Gizmo; };
+	check(ActiveGizmos.FindByPredicate(Pred));
 
 	InputRouter->ForceTerminateSource(Gizmo);
 
@@ -149,7 +140,7 @@ bool UInteractiveGizmoManager::DestroyGizmo(UInteractiveGizmo* Gizmo)
 
 	InputRouter->DeregisterSource(Gizmo);
 
-	ActiveGizmos.RemoveAt(FoundIndex);
+	ActiveGizmos.RemoveAll(Pred);
 
 	PostInvalidation();
 
@@ -270,6 +261,7 @@ FString UInteractiveGizmoManager::DefaultPlanePositionBuilderIdentifier = TEXT("
 FString UInteractiveGizmoManager::DefaultAxisAngleBuilderIdentifier = TEXT("StandardXFormAxisRotationGizmo");
 FString UInteractiveGizmoManager::DefaultThreeAxisTransformBuilderIdentifier = TEXT("DefaultThreeAxisTransformBuilderIdentifier");
 const FString UInteractiveGizmoManager::CustomThreeAxisTransformBuilderIdentifier = TEXT("CustomThreeAxisTransformBuilderIdentifier");
+FString UInteractiveGizmoManager::DefaultScalableSphereBuilderIdentifier = TEXT("DefaultScalableSphereBuilderIdentifier");
 
 void UInteractiveGizmoManager::RegisterDefaultGizmos()
 {
@@ -294,6 +286,9 @@ void UInteractiveGizmoManager::RegisterDefaultGizmos()
 	UIntervalGizmoBuilder* IntervalGizmoBuilder = NewObject<UIntervalGizmoBuilder>();
 	RegisterGizmoType(UIntervalGizmo::GizmoName, IntervalGizmoBuilder);
 
+	UScalableSphereGizmoBuilder* ScalableSphereBuilder = NewObject<UScalableSphereGizmoBuilder>();
+	RegisterGizmoType(DefaultScalableSphereBuilderIdentifier, ScalableSphereBuilder);
+
 	bDefaultGizmosRegistered = true;
 }
 
diff --git a/Engine/Source/Runtime/Experimental/InteractiveToolsFramework/Public/BaseGizmos/ScalableSphereGizmo.h b/Engine/Source/Runtime/Experimental/InteractiveToolsFramework/Public/BaseGizmos/ScalableSphereGizmo.h
new file mode 100644
--- /dev/null
+++ b/Engine/Source/Runtime/Experimental/InteractiveToolsFramework/Public/BaseGizmos/ScalableSphereGizmo.h
@@ -0,0 +1,174 @@
+// Copyright Epic Games, Inc. All Rights Reserved.
+
+#pragma once
+
+#include "CoreMinimal.h"
+#include "InteractiveGizmo.h"
+#include "InteractiveGizmoBuilder.h"
+#include "TransformProxy.h"
+#include "GizmoActor.h"
+#include "BaseBehaviors/BehaviorTargetInterfaces.h"
+#include "BaseBehaviors/AnyButtonInputBehavior.h"
+
+#include "ScalableSphereGizmo.generated.h"
+
+UCLASS()
+class INTERACTIVETOOLSFRAMEWORK_API UScalableSphereGizmoBuilder : public UInteractiveGizmoBuilder
+{
+	GENERATED_BODY()
+
+public:
+	virtual UInteractiveGizmo* BuildGizmo(const FToolBuilderState& SceneState) const override;
+};
+
+UCLASS()
+class INTERACTIVETOOLSFRAMEWORK_API AScalableSphereGizmoActor : public AGizmoActor
+{
+	GENERATED_BODY()
+
+	public:
+
+	AScalableSphereGizmoActor();
+
+	// +X Axis Handle
+	UPROPERTY()
+	UPrimitiveComponent* XPositive;
+
+	// -X Axis Handle
+	UPROPERTY()
+	UPrimitiveComponent* XNegative;
+
+	// +Y Axis Handle
+	UPROPERTY()
+	UPrimitiveComponent* YPositive;
+
+	// -Y Axis Handle
+	UPROPERTY()
+	UPrimitiveComponent* YNegative;
+
+	// +Z Axis Handle
+	UPROPERTY()
+	UPrimitiveComponent* ZPositive;
+
+	// -Z Axis Handle
+	UPROPERTY()
+	UPrimitiveComponent* ZNegative;
+};
+
+/**
+ * UScalableSphereGizmo provides a sphere that can be scaled in all directions by dragging it's handles
+ *
+ * The in-scene representation of the Gizmo is a AScalableSphereGizmoActor (or subclass).
+ * This Actor has FProperty members for the various sub-widgets, each as a separate Component.
+ * Any particular sub-widget of the Gizmo can be disabled by setting the respective
+ * Actor Component to null.
+ *
+ */
+UCLASS()
+class INTERACTIVETOOLSFRAMEWORK_API UScalableSphereGizmo : public UInteractiveGizmo, public IHoverBehaviorTarget
+{
+	GENERATED_BODY()
+
+public:
+	// UInteractiveGizmo interface
+
+	virtual void Setup() override;
+
+	virtual void Render(IToolsContextRenderAPI* RenderAPI) override;
+
+	virtual void Shutdown() override;
+
+	// IHoverBehaviorTarget interface
+	// TODO: Add HoverBehavior
+	virtual FInputRayHit BeginHoverSequenceHitTest(const FInputDeviceRay& PressPos) override { return FInputRayHit(); }
+	virtual void OnBeginHover(const FInputDeviceRay& DevicePos) override {}
+	virtual bool OnUpdateHover(const FInputDeviceRay& DevicePos) override { return true; }
+	virtual void OnEndHover() override {}
+	
+	/**
+	 * Set the Target to which the gizmo will be attached
+	 */
+	virtual void SetTarget(UTransformProxy* InTarget);
+
+	/* Set the World in which the GizmoActor will be spawned*/
+	virtual void SetWorld(UWorld* InWorld);
+
+	virtual void OnBeginDrag(const FInputDeviceRay& Ray);
+	virtual void OnUpdateDrag(const FInputDeviceRay& Ray);
+
+	/** Check if the input Ray hit any of the components of the internal actor */
+	bool HitTest(const FRay& Ray, FHitResult& OutHit, FVector& OutAxis, FTransform& OutTransform);
+
+	/* Set the Radius of the Sphere*/
+	void SetRadius(float InRadius);
+
+	/** Called when the radius is chaged (by dragging or setting). Sends new radius as parameter. */
+	TFunction<void(const float)> UpdateRadiusFunc = nullptr;
+private:
+
+	/**
+	 * Create the GizmoActor and the handles
+	 */
+	void CreateGizmoHandles();
+
+	/**
+	 * Update the handles to the right position when the radius changes
+	 */
+	void UpdateGizmoHandles();
+
+	/**
+	 * Callback for when the ActiveTarget's transform is changed
+	 */
+	void OnTransformChanged(UTransformProxy*, FTransform);
+
+	// The radius of the sphere
+	UPROPERTY()
+	float Radius;
+
+	UPROPERTY()
+	UWorld* World;
+
+	UPROPERTY()
+	UTransformProxy* ActiveTarget;
+
+	UPROPERTY()
+	AScalableSphereGizmoActor* GizmoActor;
+
+	// The current axis that is being dragged along
+	UPROPERTY()
+	FVector ActiveAxis;
+
+	// The position the drag was started on
+	UPROPERTY()
+	FVector DragStartWorldPosition;
+
+	// The initial parameter along the drag axist
+	UPROPERTY()
+	float InteractionStartParameter;
+};
+
+
+/**
+ * A behavior that forwards clicking and dragging to the gizmo.
+ */
+UCLASS()
+class INTERACTIVETOOLSFRAMEWORK_API UScalableSphereGizmoInputBehavior : public UAnyButtonInputBehavior
+{
+	GENERATED_BODY()
+
+public:
+	virtual FInputCapturePriority GetPriority() override { return FInputCapturePriority(FInputCapturePriority::DEFAULT_GIZMO_PRIORITY); }
+
+	virtual void Initialize(UScalableSphereGizmo* Gizmo);
+
+	virtual FInputCaptureRequest WantsCapture(const FInputDeviceState& input) override;
+	virtual FInputCaptureUpdate BeginCapture(const FInputDeviceState& input, EInputCaptureSide eSide) override;
+	virtual FInputCaptureUpdate UpdateCapture(const FInputDeviceState& input, const FInputCaptureData& data) override;
+	virtual void ForceEndCapture(const FInputCaptureData& data) override;
+
+protected:
+	UScalableSphereGizmo* Gizmo;
+	FRay LastWorldRay;
+	FVector2D LastScreenPosition;
+	bool bInputDragCaptured;
+};
\ No newline at end of file
diff --git a/Engine/Source/Runtime/Experimental/InteractiveToolsFramework/Public/InteractiveGizmoManager.h b/Engine/Source/Runtime/Experimental/InteractiveToolsFramework/Public/InteractiveGizmoManager.h
--- a/Engine/Source/Runtime/Experimental/InteractiveToolsFramework/Public/InteractiveGizmoManager.h
+++ b/Engine/Source/Runtime/Experimental/InteractiveToolsFramework/Public/InteractiveGizmoManager.h
@@ -217,7 +217,7 @@ public:
 	static FString DefaultAxisAngleBuilderIdentifier;
 	static FString DefaultThreeAxisTransformBuilderIdentifier;
 	static const FString CustomThreeAxisTransformBuilderIdentifier;
-
+	static FString DefaultScalableSphereBuilderIdentifier;
 
 protected:
 	/** set of Currently-active Gizmos */
