Updated the in-viewport light gizmos to improve their functionality and add some more polish

diff --git a/Engine/Plugins/Experimental/GizmoEdMode/Source/LightGizmos/Private/DirectionalLightGizmo.cpp b/Engine/Plugins/Experimental/GizmoEdMode/Source/LightGizmos/Private/DirectionalLightGizmo.cpp
--- a/Engine/Plugins/Experimental/GizmoEdMode/Source/LightGizmos/Private/DirectionalLightGizmo.cpp
+++ b/Engine/Plugins/Experimental/GizmoEdMode/Source/LightGizmos/Private/DirectionalLightGizmo.cpp
@@ -8,6 +8,8 @@
 #include "BaseGizmos/GizmoMath.h"
 #include "BaseGizmos/GizmoCircleComponent.h"
 #include "BaseGizmos/GizmoLineHandleComponent.h"
+#include "BaseGizmos/GizmoRenderingUtil.h"
+#include "SceneManagement.h"
 
 // UDirectionalLightGizmoBuilder
 
@@ -60,7 +62,48 @@ void UDirectionalLightGizmo::Setup()
 
 void UDirectionalLightGizmo::Render(IToolsContextRenderAPI* RenderAPI)
 {
+	FVector Start = GizmoActor->GetActorLocation();
 	
+	// Parameters of the handle to rotate around y axis
+	FVector ArrowDir = GizmoActor->GetActorRotation().Vector();
+	FVector ArrowEnd = Start + ArrowDir * ArrowLength;
+
+	FVector ZAxis;
+	
+	// Checking if the rotation normal is positive or negative Z axis
+	if (ArrowEnd.Z < Start.Z)
+	{
+		ZAxis = -FVector::ZAxisVector;
+	}
+	else
+	{
+		ZAxis = FVector::ZAxisVector;
+	}
+
+	// Project the arrow end onto the XY plane to get it's direction along that plane
+	FVector EndProjection = GizmoMath::ProjectPointOntoPlane(ArrowEnd, Start, ZAxis);
+
+	FVector LineDir = EndProjection - Start;
+	LineDir.Normalize();
+
+	// The end point of the line
+	FVector LineEnd = Start + LineDir * ArrowLength;
+
+	const FSceneView* View = RenderAPI->GetSceneView();
+	float PixelToWorld = GizmoRenderingUtil::CalculateLocalPixelToWorldScale(View, LineEnd);
+
+	// Calculate the "true" end point in world space (We want the line to be the same length in screen space)
+	FVector LineEndWorld = Start + LineDir * ArrowLength * PixelToWorld;
+
+	// Draw the line that represents the arrow projected onto the circle
+	RenderAPI->GetPrimitiveDrawInterface()->DrawLine(Start, LineEndWorld, FLinearColor::Red, SDPG_Foreground);
+
+	// Figure out the angle of the arc between the line and the arrow
+	float ArcAngle = FMath::Acos(FVector::DotProduct(LineDir, ArrowDir));
+	ArcAngle = FMath::RadiansToDegrees(ArcAngle);
+
+	// Draw an arc between the arrow and the line
+	DrawArc(RenderAPI->GetPrimitiveDrawInterface(), Start, LineDir, ZAxis, 0, ArcAngle, ArrowLength * PixelToWorld, 32, FLinearColor(0.f, 0.f, 1.f), SDPG_Foreground);
 }
 
 void UDirectionalLightGizmo::Shutdown()
@@ -266,6 +309,7 @@ void UDirectionalLightGizmo::CreateZRotationGizmo()
 	NewCircle->Normal = FVector::ZAxisVector;
 	NewCircle->Color = FLinearColor::Blue;
 	NewCircle->Radius = 120.f;
+	NewCircle->bDrawFullCircle = true;
 	NewCircle->RegisterComponent();
 
 	GizmoActor->RotationZCircle = NewCircle;
diff --git a/Engine/Plugins/Experimental/GizmoEdMode/Source/LightGizmos/Private/ScalableConeGizmo.cpp b/Engine/Plugins/Experimental/GizmoEdMode/Source/LightGizmos/Private/ScalableConeGizmo.cpp
--- a/Engine/Plugins/Experimental/GizmoEdMode/Source/LightGizmos/Private/ScalableConeGizmo.cpp
+++ b/Engine/Plugins/Experimental/GizmoEdMode/Source/LightGizmos/Private/ScalableConeGizmo.cpp
@@ -5,28 +5,16 @@
 #include "BaseGizmos/GizmoBoxComponent.h"
 #include "Engine/CollisionProfile.h"
 #include "BaseGizmos/GizmoMath.h"
+#include "BaseGizmos/GizmoRenderingUtil.h"
+#include "BaseBehaviors/MouseHoverBehavior.h"
 
 // UScalableConeGizmoBuilder
 
 UInteractiveGizmo* UScalableConeGizmoBuilder::BuildGizmo(const FToolBuilderState& SceneState) const
 {
 	UScalableConeGizmo* NewGizmo = NewObject<UScalableConeGizmo>(SceneState.GizmoManager);
-	NewGizmo->SetWorld(SceneState.World);
 	return NewGizmo;
 }
-
-// AScalableConeGizmoActor
-
-AScalableConeGizmoActor::AScalableConeGizmoActor()
-{
-	// root component is a hidden sphere
-	USphereComponent* SphereComponent = CreateDefaultSubobject<USphereComponent>(TEXT("GizmoCenter"));
-	RootComponent = SphereComponent;
-	SphereComponent->InitSphereRadius(1.0f);
-	SphereComponent->SetVisibility(false);
-	SphereComponent->SetCollisionProfileName(UCollisionProfile::NoCollision_ProfileName);
-}
-
 // UScalableConeGizmo
 
 void UScalableConeGizmo::Setup()
@@ -43,40 +31,64 @@ void UScalableConeGizmo::Setup()
 	ScalableConeBehavior->Initialize(this);
 	AddInputBehavior(ScalableConeBehavior);
 
-	CreateGizmoHandles();
+	UMouseHoverBehavior* HoverBehavior = NewObject<UMouseHoverBehavior>(this);
+	HoverBehavior->Initialize(this);
+	AddInputBehavior(HoverBehavior);
 }
 
 void UScalableConeGizmo::Render(IToolsContextRenderAPI* RenderAPI)
 {
 	if (ActiveTarget)
 	{
-		DrawWireSphereCappedCone(RenderAPI->GetPrimitiveDrawInterface(), GizmoActor->GetTransform(), Length, Angle, 32, 8, 10, ConeColor, SDPG_World);
-	}
-}
+		DrawWireSphereCappedCone(RenderAPI->GetPrimitiveDrawInterface(), ActiveTarget->GetTransform(), Length, Angle, 32, 8, 10, ConeColor, SDPG_Foreground);
 
-void UScalableConeGizmo::Shutdown()
-{
-	if (GizmoActor)
-	{
-		GizmoActor->Destroy();
-		GizmoActor = nullptr;
-	}
-}
+		if (bIsHovering)
+		{
+			FVector WorldOrigin = ActiveTarget->GetTransform().GetLocation();
+			FVector CircleNormal = ActiveTarget->GetTransform().GetRotation().Vector();
 
-void UScalableConeGizmo::SetTarget(UTransformProxy* InTarget)
-{
-	ActiveTarget = InTarget;
+			float ConeHeight = Length * FMath::Cos(FMath::DegreesToRadians(Angle));
+			float ConeRadius = Length * FMath::Sin(FMath::DegreesToRadians(Angle));
+
+			// The center of the circle at the base of the cone
+			FVector CircleOrigin = WorldOrigin + CircleNormal * ConeHeight;
+
+			FVector CircleX;
+			FVector CircleY;
+
+			// Direction vectors to represent the circle
+			GizmoMath::MakeNormalPlaneBasis(CircleNormal, CircleX, CircleY);
+
+			float CircleThickness = HoverThickness;
 
-	// To update the internal GizmoActor when the transform is changed
-	ActiveTarget->OnTransformChanged.AddUObject(this, &UScalableConeGizmo::OnTransformChanged);
+			// Figure out the correct thickness of the Circle in world coordinates
+			const FSceneView* View = RenderAPI->GetSceneView();
 
-	OnTransformChanged(InTarget, InTarget->GetTransform());
+			FVector CircleEndWithThickness = DragStartWorldPosition + CircleX * CircleThickness;
+			float ScreenThickness = GizmoRenderingUtil::CalculateLocalPixelToWorldScale(View, CircleEndWithThickness);
 
+			CircleThickness *= ScreenThickness;
+
+			DrawCircle(RenderAPI->GetPrimitiveDrawInterface(), CircleOrigin, CircleX, CircleY, ConeColor, ConeRadius, 32, SDPG_Foreground, CircleThickness);
+
+			// Parameters for the line that shows the drag direction
+			FVector LineStart = DragStartWorldPosition;
+			float LineLength = 30.f;
+			FVector LineEnd = LineStart + HitAxis * LineLength;
+
+			// Get the Pixel to World scale of the line
+			float PixelToWorld = GizmoRenderingUtil::CalculateLocalPixelToWorldScale(View, LineEnd);
+
+			// Draw the lines in both directions
+			RenderAPI->GetPrimitiveDrawInterface()->DrawLine(LineStart, LineStart + HitAxis * PixelToWorld * LineLength, FLinearColor::Red, SDPG_Foreground);
+			RenderAPI->GetPrimitiveDrawInterface()->DrawLine(LineStart, LineStart - HitAxis * PixelToWorld * LineLength, FLinearColor::Red, SDPG_Foreground);
+		}
+	}
 }
 
-void UScalableConeGizmo::SetWorld(UWorld* InWorld)
+void UScalableConeGizmo::SetTarget(UTransformProxy* InTarget)
 {
-	World = InWorld;
+	ActiveTarget = InTarget;
 }
 
 void UScalableConeGizmo::SetAngleDegrees(float InAngle)
@@ -85,8 +97,6 @@ void UScalableConeGizmo::SetAngleDegrees(float InAngle)
 
 	Angle = FMath::Clamp<float>(Angle, MinAngle, MaxAngle);
 
-	UpdateGizmoHandles();
-
 	if (UpdateAngleFunc)
 	{
 		UpdateAngleFunc(InAngle);
@@ -108,44 +118,55 @@ float UScalableConeGizmo::GetAngleDegrees()
 	return Angle;
 }
 
-void UScalableConeGizmo::CreateGizmoHandles()
+FInputRayHit UScalableConeGizmo::BeginHoverSequenceHitTest(const FInputDeviceRay& PressPos)
 {
-	FActorSpawnParameters SpawnInfo;
-	GizmoActor = World->SpawnActor<AScalableConeGizmoActor>(FVector::ZeroVector, FRotator::ZeroRotator, SpawnInfo);
-
-	 // Handles are Boxes right now
-	GizmoActor->ScaleHandleYPlus = AGizmoActor::AddDefaultBoxComponent(World, GizmoActor, FLinearColor(1.0, 0, 0), FVector::ZeroVector);
-	GizmoActor->ScaleHandleYMinus = AGizmoActor::AddDefaultBoxComponent(World, GizmoActor, FLinearColor(1.0, 0, 0), FVector::ZeroVector);
-	GizmoActor->ScaleHandleZPlus = AGizmoActor::AddDefaultBoxComponent(World, GizmoActor, FLinearColor(1.0, 0, 0), FVector::ZeroVector);
-	GizmoActor->ScaleHandleZMinus = AGizmoActor::AddDefaultBoxComponent(World, GizmoActor, FLinearColor(1.0, 0, 0), FVector::ZeroVector);
-	
-	UpdateGizmoHandles();
-}
+	FHitResult HitResult;
+	FVector TestHitAxis;
+	FTransform DragTransform;
 
-void UScalableConeGizmo::UpdateGizmoHandles()
-{
-	 // Get radius and height of the cone
-	float Radius = Length * FMath::Sin(FMath::DegreesToRadians(Angle));
-	float Height = Length * FMath::Cos(FMath::DegreesToRadians(Angle));
-
-	GizmoActor->ScaleHandleYPlus->SetRelativeLocation(FVector::XAxisVector * Height + FVector::YAxisVector * Radius);
-	GizmoActor->ScaleHandleYMinus->SetRelativeLocation(FVector::XAxisVector * Height - FVector::YAxisVector * Radius);
-	GizmoActor->ScaleHandleZPlus->SetRelativeLocation(FVector::XAxisVector * Height + FVector::ZAxisVector * Radius);
-	GizmoActor->ScaleHandleZMinus->SetRelativeLocation(FVector::XAxisVector * Height - FVector::ZAxisVector * Radius);
+	if (HitTest(PressPos.WorldRay, HitResult, TestHitAxis, DragTransform))
+	{
+		bIsHovering = true;
+		DragStartWorldPosition = DragTransform.GetLocation();
+		HitAxis = TestHitAxis;
+
+		return FInputRayHit(HitResult.Distance);
+	}
+
+	// Return invalid ray hit to say we don't want to listen to hover input
+	return FInputRayHit();
 }
 
-void UScalableConeGizmo::OnTransformChanged(UTransformProxy*, FTransform)
+bool UScalableConeGizmo::OnUpdateHover(const FInputDeviceRay& DevicePos)
 {
-	USceneComponent* GizmoComponent = GizmoActor->GetRootComponent();
+	if (!ActiveTarget)
+	{
+		return false;
+	}
 
-	FTransform TargetTransform = ActiveTarget->GetTransform();
+	FVector Start = DevicePos.WorldRay.Origin;
+	const float MaxRaycastDistance = 1e6f;
+	FVector End = DevicePos.WorldRay.Origin + DevicePos.WorldRay.Direction * MaxRaycastDistance;
 
-	// GizmoActor doesn't want the scale of the object
-	TargetTransform.SetScale3D(FVector(1, 1, 1));
+	FRay HitCheckRay(Start, End - Start);
+	FHitResult HitResult;
+	FVector TestHitAxis;
+	FTransform DragTransform;
 
-	GizmoComponent->SetWorldTransform(TargetTransform);
+	if (HitTest(HitCheckRay, HitResult, TestHitAxis, DragTransform))
+	{
+		bIsHovering = true;
+		DragStartWorldPosition = DragTransform.GetLocation();
+		HitAxis = TestHitAxis;
+		return true;
+	}
 
-	UpdateGizmoHandles();
+	return false;
+}
+
+void UScalableConeGizmo::OnEndHover()
+{
+	bIsHovering = false;
 }
 
 bool UScalableConeGizmo::HitTest(const FRay& Ray, FHitResult& OutHit, FVector& OutAxis, FTransform& OutTransform)
@@ -159,38 +180,50 @@ bool UScalableConeGizmo::HitTest(const FRay& Ray, FHitResult& OutHit, FVector& O
 	const float MaxRaycastDistance = 1e6f;
 	FVector End = Ray.Origin + Ray.Direction * MaxRaycastDistance;
 
-	FCollisionQueryParams Params;
+	// Find the intresection with the circle plane. Note that unlike the FMath version, GizmoMath::RayPlaneIntersectionPoint() 
+	// checks that the ray isn't parallel to the plane.
+	FVector WorldOrigin = ActiveTarget->GetTransform().GetLocation();
+	FVector CircleNormal = ActiveTarget->GetTransform().GetRotation().Vector();
 
-	FQuat Rotation = ActiveTarget->GetTransform().GetRotation();
+	float ConeHeight = Length * FMath::Cos(FMath::DegreesToRadians(Angle));
+	float ConeRadius = Length * FMath::Sin(FMath::DegreesToRadians(Angle));
 
-	// Check which component was hit and update OutAxis and OutTransform
-	if (GizmoActor->ScaleHandleYPlus->LineTraceComponent(OutHit, Start, End, Params))
-	{
-		OutAxis = Rotation.RotateVector(FVector::YAxisVector);
-		OutTransform = GizmoActor->ScaleHandleYPlus->GetComponentTransform();
-		return true;
-	}
-	if (GizmoActor->ScaleHandleYMinus->LineTraceComponent(OutHit, Start, End, Params))
-	{
-		OutAxis = Rotation.RotateVector(-FVector::YAxisVector);
-		OutTransform = GizmoActor->ScaleHandleYMinus->GetComponentTransform();
-		return true;
-	}
-	if (GizmoActor->ScaleHandleZPlus->LineTraceComponent(OutHit, Start, End, Params))
+	// The center of the base of the cone
+	FVector CircleOrigin = WorldOrigin + CircleNormal * ConeHeight;
+	
+	FVector ConeBottom = WorldOrigin + CircleNormal * Length;
+
+	FVector HitPos;
+	bool bIntersects = false;
+
+	// Figure out if the ray interesects with the plane at the base of the cone
+	GizmoMath::RayPlaneIntersectionPoint(CircleOrigin, CircleNormal, Ray.Origin, Ray.Direction, bIntersects, HitPos);
+
+	if (!bIntersects || Ray.GetParameter(HitPos) > Ray.GetParameter(End))
 	{
-		OutAxis = Rotation.RotateVector(FVector::ZAxisVector);
-		OutTransform = GizmoActor->ScaleHandleZPlus->GetComponentTransform();
-		return true;
+		return false;
 	}
-	if (GizmoActor->ScaleHandleZMinus->LineTraceComponent(OutHit, Start, End, Params))
+
+	// Find the closest point on the circle of the intersection
+	FVector NearestCircle;
+	GizmoMath::ClosetPointOnCircle(HitPos, CircleOrigin, CircleNormal, ConeRadius, NearestCircle);
+
+	FVector NearestRay = Ray.ClosestPoint(NearestCircle);
+
+	// Make sure the distance to the ray is within a certain threshold
+	double Distance = FVector::Distance(NearestCircle, NearestRay);
+
+	if (Distance > HitErrorThreshold)
 	{
-		OutAxis = Rotation.RotateVector(-FVector::ZAxisVector);
-		OutTransform = GizmoActor->ScaleHandleZMinus->GetComponentTransform();
-		return true;
+		return false;
 	}
 
-	return false;
+	OutAxis = NearestCircle - ConeBottom;
+	OutAxis.Normalize();
 
+	OutTransform.SetIdentity();
+	OutTransform.SetTranslation(NearestCircle);
+	return true;
 }
 
 void UScalableConeGizmo::OnBeginDrag(const FInputDeviceRay& Ray)
@@ -202,10 +235,13 @@ void UScalableConeGizmo::OnBeginDrag(const FInputDeviceRay& Ray)
 	FRay HitCheckRay(Start, End - Start);
 	FHitResult HitResult;
 	FTransform DragTransform;
+	FVector HitTestAxis;
 
 	 // Check if any component was hit
-	if (HitTest(HitCheckRay, HitResult, HitAxis, DragTransform))
+	if (HitTest(HitCheckRay, HitResult, HitTestAxis, DragTransform))
 	{
+		HitAxis = HitTestAxis;
+
 		FVector RayNearestPt; float RayNearestParam;
 
 		// Get the initial interaction parameters
@@ -229,7 +265,7 @@ void UScalableConeGizmo::OnUpdateDrag(const FInputDeviceRay& Ray)
 		AxisNearestPt, AxisNearestParam,
 		RayNearestPt, RayNearestParam);
 
-	FVector GizmoLocation = GizmoActor->GetActorLocation();
+	FVector GizmoLocation = ActiveTarget->GetTransform().GetLocation();
 
 	// Vector to the starting position of the interaction
 	FVector Start = InteractionStartPoint - GizmoLocation;
@@ -244,8 +280,11 @@ void UScalableConeGizmo::OnUpdateDrag(const FInputDeviceRay& Ray)
 	float DeltaAngle = FMath::Acos(DotP);
 
 	// Get the angle between the start and end vectors and the forward vector to check if the drag direction should be +ve or -ve
-	float StartAngle = FMath::Acos(FVector::DotProduct(Start, GizmoActor->GetActorForwardVector()));
-	float EndAngle = FMath::Acos(FVector::DotProduct(End, GizmoActor->GetActorForwardVector()));
+
+	FVector Forward = ActiveTarget->GetTransform().GetRotation().Vector();
+
+	float StartAngle = FMath::Acos(FVector::DotProduct(Start, Forward));
+	float EndAngle = FMath::Acos(FVector::DotProduct(End, Forward));
 
 	if (StartAngle > EndAngle)
 	{
@@ -270,9 +309,9 @@ FInputCaptureRequest UScalableConeGizmoInputBehavior::WantsCapture(const FInputD
 	if (IsPressed(input))
 	{
 		FHitResult HitResult;
-		FVector HitAxis;
+		FVector HitTestAxis;
 		FTransform DragTransform;
-		if (Gizmo->HitTest(input.Mouse.WorldRay, HitResult, HitAxis, DragTransform))
+		if (Gizmo->HitTest(input.Mouse.WorldRay, HitResult, HitTestAxis, DragTransform))
 		{
 			return FInputCaptureRequest::Begin(this, EInputCaptureSide::Any, HitResult.Distance);
 		}
diff --git a/Engine/Plugins/Experimental/GizmoEdMode/Source/LightGizmos/Private/SpotLightGizmo.cpp b/Engine/Plugins/Experimental/GizmoEdMode/Source/LightGizmos/Private/SpotLightGizmo.cpp
--- a/Engine/Plugins/Experimental/GizmoEdMode/Source/LightGizmos/Private/SpotLightGizmo.cpp
+++ b/Engine/Plugins/Experimental/GizmoEdMode/Source/LightGizmos/Private/SpotLightGizmo.cpp
@@ -8,7 +8,7 @@
 #include "Kismet/KismetMathLibrary.h"
 #include "BaseGizmos/GizmoMath.h"
 #include "Components/SphereComponent.h"
-#include "BaseGizmos/GizmoBoxComponent.h"
+#include "BaseGizmos/GizmoLineHandleComponent.h"
 
 // USpotLightGizmoBuilder
 
@@ -228,6 +228,8 @@ void USpotLightGizmo::OnUpdateDrag(const FInputDeviceRay& Ray)
 	// Update the attenuation of the cone
 	float NewAttenuation = LightActor->SpotLightComponent->AttenuationRadius + DeltaParam;
 
+	NewAttenuation = (NewAttenuation < 0) ? 0 : NewAttenuation;
+
 	LightActor->SpotLightComponent->AttenuationRadius = NewAttenuation;
 	LightActor->SpotLightComponent->MarkRenderStateDirty();
 
@@ -274,8 +276,8 @@ void USpotLightGizmo::CreateAttenuationScaleGizmo()
 	FActorSpawnParameters SpawnInfo;
 	GizmoActor = World->SpawnActor<ASpotLightGizmoActor>(FVector::ZeroVector, FRotator::ZeroRotator, SpawnInfo);
 
-	// The handle to scale attenuation is a box 
-	GizmoActor->AttenuationScaleHandle = AGizmoActor::AddDefaultBoxComponent(World, GizmoActor, FLinearColor::Red, FVector::ZeroVector);
+	// The handle to scale attenuation is line handle component
+	GizmoActor->AttenuationScaleHandle = AGizmoActor::AddDefaultLineHandleComponent(World, GizmoActor, FLinearColor::Blue, FVector::YAxisVector, FVector::XAxisVector, 60.f, true);
 	GizmoActor->AttenuationScaleHandle->SetRelativeLocation(FVector(LightActor->SpotLightComponent->AttenuationRadius, 0, 0));
 
 	TransformProxy->OnTransformChanged.AddUObject(this, &USpotLightGizmo::OnTransformChanged);
diff --git a/Engine/Plugins/Experimental/GizmoEdMode/Source/LightGizmos/Public/ScalableConeGizmo.h b/Engine/Plugins/Experimental/GizmoEdMode/Source/LightGizmos/Public/ScalableConeGizmo.h
--- a/Engine/Plugins/Experimental/GizmoEdMode/Source/LightGizmos/Public/ScalableConeGizmo.h
+++ b/Engine/Plugins/Experimental/GizmoEdMode/Source/LightGizmos/Public/ScalableConeGizmo.h
@@ -21,26 +21,9 @@ public:
 	virtual UInteractiveGizmo* BuildGizmo(const FToolBuilderState& SceneState) const override;
 };
 
-UCLASS()
-class LIGHTGIZMOS_API AScalableConeGizmoActor : public AGizmoActor
-{
-	GENERATED_BODY()
-
-public:
-
-	AScalableConeGizmoActor();
-
-	UPrimitiveComponent* ScaleHandleYPlus;
-	UPrimitiveComponent* ScaleHandleYMinus;
-	UPrimitiveComponent* ScaleHandleZPlus;
-	UPrimitiveComponent* ScaleHandleZMinus;
-
-	UPrimitiveComponent* LengthHandle;
-};
-
 /**
  * UScalableConeGizmo provides a cone that can be scaled (changing its angle)
- * The in-scene representation of the Gizmo is a AScalableConeGizmoActor (or subclass).
+ * by dragging the base of the cone outwards/inwards
  */
 UCLASS()
 class LIGHTGIZMOS_API UScalableConeGizmo : public UInteractiveGizmo, public IHoverBehaviorTarget
@@ -54,11 +37,8 @@ public:
 
 	virtual void Render(IToolsContextRenderAPI* RenderAPI) override;
 
-	virtual void Shutdown() override;
-
 	// Set the target to attach the gizmo to
 	virtual void SetTarget(UTransformProxy* InTarget);
-	virtual void SetWorld(UWorld* InWorld);
 
 	// Gettors and Settors for the Angle and Length
 	void SetAngleDegrees(float InAngle);
@@ -67,13 +47,17 @@ public:
 	float GetAngleDegrees();
 
 	// IHoverBehaviorTarget interface
-	virtual FInputRayHit BeginHoverSequenceHitTest(const FInputDeviceRay& PressPos) override { return FInputRayHit(); }
+	virtual FInputRayHit BeginHoverSequenceHitTest(const FInputDeviceRay& PressPos) override;
 	virtual void OnBeginHover(const FInputDeviceRay& DevicePos) override {}
-	virtual bool OnUpdateHover(const FInputDeviceRay& DevicePos) override { return true; }
-	virtual void OnEndHover() override {}
+	virtual bool OnUpdateHover(const FInputDeviceRay& DevicePos) override;
+	virtual void OnEndHover() override;
 
 	virtual void OnBeginDrag(const FInputDeviceRay& Ray);
 	virtual void OnUpdateDrag(const FInputDeviceRay& Ray);
+
+	/** Check if the input ray hits the cone. The hit testable parts of the cone currently is just 
+	 *  the circle at its base
+	 */
 	bool HitTest(const FRay& Ray, FHitResult& OutHit, FVector& OutAxis, FTransform& OutTransform);
 
 	// The maximum angle the cone can be stretched to
@@ -88,15 +72,19 @@ public:
 	UPROPERTY()
 	FColor ConeColor;
 
+	// The error threshold for hit detection with the cone
+	UPROPERTY()
+	float HitErrorThreshold{ 12.f };
+
+	// The thickness of the base of the cone when being hovered over (in screen space)
+	UPROPERTY()
+	float HoverThickness{ 5.f };
+
 	/** Called when the Angle of the cone is changed. Sends new angle as parameter. */
 	TFunction<void(const float)> UpdateAngleFunc = nullptr;
 
 private:
 
-	void CreateGizmoHandles();
-	void UpdateGizmoHandles();
-	void OnTransformChanged(UTransformProxy*, FTransform);
-
 	// The ConeLength
 	UPROPERTY()
 	float Length;
@@ -104,14 +92,12 @@ private:
 	UPROPERTY()
 	float Angle;
 
+	/** Whether the gizmo is being hovered over */
 	UPROPERTY()
-	UTransformProxy* ActiveTarget;
+	bool bIsHovering{ false };
 
 	UPROPERTY()
-	AScalableConeGizmoActor* GizmoActor;
-
-	UPROPERTY()
-	UWorld* World;
+	UTransformProxy* ActiveTarget;
 
 	/** Used for calculations when moving the handles*/
 
diff --git a/Engine/Source/Runtime/Experimental/InteractiveToolsFramework/Private/BaseGizmos/GizmoCircleComponent.cpp b/Engine/Source/Runtime/Experimental/InteractiveToolsFramework/Private/BaseGizmos/GizmoCircleComponent.cpp
--- a/Engine/Source/Runtime/Experimental/InteractiveToolsFramework/Private/BaseGizmos/GizmoCircleComponent.cpp
+++ b/Engine/Source/Runtime/Experimental/InteractiveToolsFramework/Private/BaseGizmos/GizmoCircleComponent.cpp
@@ -23,6 +23,7 @@ public:
 		Thickness(InComponent->Thickness),
 		NumSides(InComponent->NumSides),
 		bViewAligned(InComponent->bViewAligned),
+		bDrawFullCircle(InComponent->bDrawFullCircle),
 		HoverThicknessMultiplier(InComponent->HoverSizeMultiplier)
 	{
 	}
@@ -72,6 +73,12 @@ public:
 					BackThickness *= (View->FOV / 90.0);		// compensate for FOV scaling in Gizmos...
 				}
 
+				if (bDrawFullCircle)
+				{
+					BackThickness = UseThickness;
+					BackColor = Color;
+				}
+
 				const float	AngleDelta = 2.0f * PI / NumSides;
 
 				if (bViewAligned)
@@ -219,6 +226,7 @@ private:
 	float Thickness;
 	int NumSides;
 	bool bViewAligned;
+	bool bDrawFullCircle;
 	float HoverThicknessMultiplier;
 
 	// set on Component for use in ::GetDynamicMeshElements()
diff --git a/Engine/Source/Runtime/Experimental/InteractiveToolsFramework/Private/BaseGizmos/ScalableSphereGizmo.cpp b/Engine/Source/Runtime/Experimental/InteractiveToolsFramework/Private/BaseGizmos/ScalableSphereGizmo.cpp
--- a/Engine/Source/Runtime/Experimental/InteractiveToolsFramework/Private/BaseGizmos/ScalableSphereGizmo.cpp
+++ b/Engine/Source/Runtime/Experimental/InteractiveToolsFramework/Private/BaseGizmos/ScalableSphereGizmo.cpp
@@ -8,31 +8,17 @@
 #include "Engine/CollisionProfile.h"
 #include "Kismet/KismetMathLibrary.h"
 #include "BaseGizmos/GizmoMath.h"
+#include "BaseBehaviors/MouseHoverBehavior.h"
+#include "BaseGizmos/GizmoRenderingUtil.h"
 
 // UScalableSphereGizmoBuilder
 
 UInteractiveGizmo* UScalableSphereGizmoBuilder::BuildGizmo(const FToolBuilderState& SceneState) const
 {
 	UScalableSphereGizmo* NewGizmo = NewObject<UScalableSphereGizmo>(SceneState.GizmoManager);
-
-	// Have to set world to be able to spawn actor
-	NewGizmo->SetWorld(SceneState.World);
-	
 	return NewGizmo;
 }
 
-// AScalableSphereGizmoActor
-
-AScalableSphereGizmoActor::AScalableSphereGizmoActor()
-{
-	// root component is a hidden sphere
-	USphereComponent* SphereComponent = CreateDefaultSubobject<USphereComponent>(TEXT("GizmoCenter"));
-	RootComponent = SphereComponent;
-	SphereComponent->InitSphereRadius(1.0f);
-	SphereComponent->SetVisibility(false);
-	SphereComponent->SetCollisionProfileName(UCollisionProfile::NoCollision_ProfileName);
-}
-
 // UScalableSphereGizmo
 
 void UScalableSphereGizmo::Setup()
@@ -45,83 +31,141 @@ void UScalableSphereGizmo::Setup()
 	ScalableSphereBehavior->Initialize(this);
 	AddInputBehavior(ScalableSphereBehavior);
 
-	CreateGizmoHandles();
+	UMouseHoverBehavior* HoverBehavior = NewObject<UMouseHoverBehavior>(this);
+	HoverBehavior->Initialize(this);
+	AddInputBehavior(HoverBehavior);
 }
 
 void UScalableSphereGizmo::Render(IToolsContextRenderAPI* RenderAPI)
 {
 	if (ActiveTarget)
 	{
-		DrawWireSphereAutoSides(RenderAPI->GetPrimitiveDrawInterface(), ActiveTarget->GetTransform(), FColor(200, 255, 255), Radius, SDPG_World);
+		float UseThickness = Thickness;
+
+		if (bIsHovering)
+		{
+			UseThickness = HoverThickness;
+
+			// Parameters for the line that shows the drag direction
+			FVector LineStart = DragStartWorldPosition;
+			float LineLength = 30.f;
+			FVector LineEnd = LineStart + ActiveAxis * LineLength;
+
+			// Get the Pixel to World scale of the line
+			const FSceneView* View = RenderAPI->GetSceneView();
+			float PixelToWorld = GizmoRenderingUtil::CalculateLocalPixelToWorldScale(View, LineEnd);
+
+			// Draw the lines in both directions
+			RenderAPI->GetPrimitiveDrawInterface()->DrawLine(LineStart, LineStart + ActiveAxis * PixelToWorld * LineLength, FLinearColor::Red, SDPG_Foreground);
+			RenderAPI->GetPrimitiveDrawInterface()->DrawLine(LineStart, LineStart - ActiveAxis * PixelToWorld* LineLength, FLinearColor::Red, SDPG_Foreground);
+
+			// Figure out the correct thickness of the sphere in world coordinates
+			FVector CircleEndWithThickness = LineStart + ActiveAxis * UseThickness;
+			float ScreenThickness = GizmoRenderingUtil::CalculateLocalPixelToWorldScale(View, CircleEndWithThickness);
+
+			UseThickness *= ScreenThickness;
+		}
+
+		DrawWireSphereAutoSides(RenderAPI->GetPrimitiveDrawInterface(), ActiveTarget->GetTransform(), FColor(200, 255, 255), Radius, SDPG_Foreground, UseThickness);
 	}
 }
 
-void UScalableSphereGizmo::Shutdown()
+FInputRayHit UScalableSphereGizmo::BeginHoverSequenceHitTest(const FInputDeviceRay& PressPos)
 {
-	if (GizmoActor)
+	FHitResult HitResult;
+	FVector HitAxis;
+	FTransform DragTransform;
+
+	if (HitTest(PressPos.WorldRay, HitResult, HitAxis, DragTransform))
 	{
-		GizmoActor->Destroy();
-		GizmoActor = nullptr;
+		bIsHovering = true;
+		DragStartWorldPosition = DragTransform.GetLocation();
+		ActiveAxis = HitAxis;
+		
+		return FInputRayHit(HitResult.Distance);
 	}
+	
+	// Return invalid ray hit to say we don't want to listen to hover input
+	return FInputRayHit();
 }
 
-void UScalableSphereGizmo::SetTarget(UTransformProxy* InTarget)
+bool UScalableSphereGizmo::OnUpdateHover(const FInputDeviceRay& DevicePos)
 {
-	ActiveTarget = InTarget;
+	if (!ActiveTarget)
+	{
+		return false;
+	}
 
-	// Make sure the internal GizmoActor updates when the transform of the target is changed
-	ActiveTarget->OnTransformChanged.AddUObject(this, &UScalableSphereGizmo::OnTransformChanged);
-	
-	OnTransformChanged(InTarget, InTarget->GetTransform());
+	FVector Start = DevicePos.WorldRay.Origin;
+	const float MaxRaycastDistance = 1e6f;
+	FVector End = DevicePos.WorldRay.Origin + DevicePos.WorldRay.Direction * MaxRaycastDistance;
+
+	FRay HitCheckRay(Start, End - Start);
+	FHitResult HitResult;
+	FVector HitAxis;
+	FTransform DragTransform;
+
+	if (HitTest(HitCheckRay, HitResult, HitAxis, DragTransform))
+	{
+		bIsHovering = true;
+		DragStartWorldPosition = DragTransform.GetLocation();
+		ActiveAxis = HitAxis;
+		return true;
+	}
+
+	return false;
 }
 
-void UScalableSphereGizmo::SetWorld(UWorld* InWorld)
+void UScalableSphereGizmo::OnEndHover()
 {
-	World = InWorld;
-
+	bIsHovering = false;
 }
 
-void UScalableSphereGizmo::CreateGizmoHandles()
+void UScalableSphereGizmo::SetTarget(UTransformProxy* InTarget)
 {
-	FActorSpawnParameters SpawnInfo;
-	GizmoActor = World->SpawnActor<AScalableSphereGizmoActor>(FVector::ZeroVector, FRotator::ZeroRotator, SpawnInfo);
+	ActiveTarget = InTarget;
+}
 
-	// Create all the 6 handles in each direction
+bool UScalableSphereGizmo::CheckCircleIntersection(const FRay& Ray, FVector CircleNormal, FVector& OutHitLocation, FVector& OutHitAxis)
+{
+	// Find the intresection with the circle plane. Note that unlike the FMath version, GizmoMath::RayPlaneIntersectionPoint() 
+	// checks that the ray isn't parallel to the plane.
+	FVector WorldOrigin = ActiveTarget->GetTransform().GetLocation();
 
-	auto NewHandle = [this] { return AGizmoActor::AddDefaultBoxComponent(World, GizmoActor, FLinearColor(1.0f, 0.0f, 0.0f), FVector::ZeroVector); };
-	GizmoActor->XPositive = NewHandle();
-	GizmoActor->XNegative = NewHandle();
-	GizmoActor->YPositive = NewHandle();
-	GizmoActor->YNegative = NewHandle();
-	GizmoActor->ZPositive = NewHandle();
-	GizmoActor->ZNegative = NewHandle();
+	FVector Start = Ray.Origin;
+	const float MaxRaycastDistance = 1e6f;
+	FVector End = Ray.Origin + Ray.Direction * MaxRaycastDistance;
 
-	UpdateGizmoHandles();
+	FVector HitPos;
+	bool bIntersects = false;
 
-}
+	// Figure out if the ray intersects the circle plane
+	GizmoMath::RayPlaneIntersectionPoint(WorldOrigin, CircleNormal, Ray.Origin, Ray.Direction, bIntersects, HitPos);
 
-void UScalableSphereGizmo::UpdateGizmoHandles()
-{
-	GizmoActor->XPositive->SetRelativeLocation(FVector::XAxisVector * Radius);
-	GizmoActor->XNegative->SetRelativeLocation(-FVector::XAxisVector * Radius);
+	if (!bIntersects || Ray.GetParameter(HitPos) > Ray.GetParameter(End))
+	{
+		return false;
+	}
 
-	GizmoActor->YPositive->SetRelativeLocation(FVector::YAxisVector * Radius);
-	GizmoActor->YNegative->SetRelativeLocation(-FVector::YAxisVector * Radius);
+	// Find the point on the circle closest to the intersection
+	FVector NearestCircle;
+	GizmoMath::ClosetPointOnCircle(HitPos, WorldOrigin, CircleNormal, Radius, NearestCircle);
 
-	GizmoActor->ZPositive->SetRelativeLocation(FVector::ZAxisVector * Radius);
-	GizmoActor->ZNegative->SetRelativeLocation(-FVector::ZAxisVector * Radius);
-}
+	FVector NearestRay = Ray.ClosestPoint(NearestCircle);
 
-void UScalableSphereGizmo::OnTransformChanged(UTransformProxy*, FTransform)
-{
-	USceneComponent* GizmoComponent = GizmoActor->GetRootComponent();
+	// Make sure the point is within a certain distance of the ray
+	double Distance = FVector::Distance(NearestCircle, NearestRay);
 
-	FTransform TargetTransform = ActiveTarget->GetTransform();
+	if (Distance > HitErrorThreshold)
+	{
+		return false;
+	}
 
-	// Don't scale the internal gizmo actor (only update position and rotation changes)
-	TargetTransform.SetScale3D(FVector(1, 1, 1));
+	OutHitAxis = NearestCircle - WorldOrigin;
+	OutHitAxis.Normalize();
 
-	GizmoComponent->SetWorldTransform(TargetTransform);
+	OutHitLocation = NearestCircle;
+	return true;
 }
 
 bool UScalableSphereGizmo::HitTest(const FRay& Ray, FHitResult& OutHit, FVector& OutAxis, FTransform &OutTransform)
@@ -131,49 +175,33 @@ bool UScalableSphereGizmo::HitTest(const FRay& Ray, FHitResult& OutHit, FVector&
 		return false;
 	}
 
-	FVector Start = Ray.Origin;
-	const float MaxRaycastDistance = 1e6f;
-	FVector End = Ray.Origin + Ray.Direction * MaxRaycastDistance;
-
-	FCollisionQueryParams Params;
+	FVector OutHitLocation;
+	FVector OutHitAxis;
 
-	// Check each actor to see if any of them were hit
-	if (GizmoActor->XPositive->LineTraceComponent(OutHit, Start, End, Params))
+	if (CheckCircleIntersection(Ray, FVector::XAxisVector, OutHitLocation, OutHitAxis))
 	{
-		OutAxis = FVector::XAxisVector;
-		OutTransform = GizmoActor->XPositive->GetComponentTransform();
-		return true;
-	}
-	if (GizmoActor->XNegative->LineTraceComponent(OutHit, Start, End, Params))
-	{
-		OutAxis = -FVector::XAxisVector;
-		OutTransform = GizmoActor->XNegative->GetComponentTransform();
-		return true;
-	}
+		OutAxis = OutHitAxis;
+		OutTransform.SetIdentity();
+		OutTransform.SetTranslation(OutHitLocation);
 
-	if (GizmoActor->YPositive->LineTraceComponent(OutHit, Start, End, Params))
-	{
-		OutAxis = FVector::YAxisVector;
-		OutTransform = GizmoActor->YPositive->GetComponentTransform();
-		return true;
-	}
-	if (GizmoActor->YNegative->LineTraceComponent(OutHit, Start, End, Params))
-	{
-		OutAxis = -FVector::YAxisVector;
-		OutTransform = GizmoActor->YNegative->GetComponentTransform();
 		return true;
 	}
 
-	if (GizmoActor->ZPositive->LineTraceComponent(OutHit, Start, End, Params))
+	if (CheckCircleIntersection(Ray, FVector::YAxisVector, OutHitLocation, OutHitAxis))
 	{
-		OutAxis = FVector::ZAxisVector;
-		OutTransform = GizmoActor->ZPositive->GetComponentTransform();
+		OutAxis = OutHitAxis;
+		OutTransform.SetIdentity();
+		OutTransform.SetTranslation(OutHitLocation);
+
 		return true;
 	}
-	if (GizmoActor->ZNegative->LineTraceComponent(OutHit, Start, End, Params))
+
+	if (CheckCircleIntersection(Ray, FVector::ZAxisVector, OutHitLocation, OutHitAxis))
 	{
-		OutAxis = -FVector::ZAxisVector;
-		OutTransform = GizmoActor->ZNegative->GetComponentTransform();
+		OutAxis = OutHitAxis;
+		OutTransform.SetIdentity();
+		OutTransform.SetTranslation(OutHitLocation);
+
 		return true;
 	}
 
@@ -194,8 +222,6 @@ void UScalableSphereGizmo::SetRadius(float InRadius)
 	{
 		UpdateRadiusFunc(Radius);
 	}
-
-	UpdateGizmoHandles();
 }
 
 void UScalableSphereGizmo::OnBeginDrag(const FInputDeviceRay& Ray)
diff --git a/Engine/Source/Runtime/Experimental/InteractiveToolsFramework/Public/BaseGizmos/GizmoCircleComponent.h b/Engine/Source/Runtime/Experimental/InteractiveToolsFramework/Public/BaseGizmos/GizmoCircleComponent.h
--- a/Engine/Source/Runtime/Experimental/InteractiveToolsFramework/Public/BaseGizmos/GizmoCircleComponent.h
+++ b/Engine/Source/Runtime/Experimental/InteractiveToolsFramework/Public/BaseGizmos/GizmoCircleComponent.h
@@ -32,6 +32,9 @@ public:
 	UPROPERTY(EditAnywhere, Category = Options)
 	bool bViewAligned = false;
 
+	UPROPERTY(EditAnywhere, Category = Options)
+	bool bDrawFullCircle = false;
+
 	// If the circle was on a 3D sphere, then only the 'front' part of the circle can be hit, 
 	// (in other words, if the ray would hit the sphere first, ignore the hit).
 	// Dynamically disabled if the circle is parallel to the view plane (ie "fully visible")
diff --git a/Engine/Source/Runtime/Experimental/InteractiveToolsFramework/Public/BaseGizmos/ScalableSphereGizmo.h b/Engine/Source/Runtime/Experimental/InteractiveToolsFramework/Public/BaseGizmos/ScalableSphereGizmo.h
--- a/Engine/Source/Runtime/Experimental/InteractiveToolsFramework/Public/BaseGizmos/ScalableSphereGizmo.h
+++ b/Engine/Source/Runtime/Experimental/InteractiveToolsFramework/Public/BaseGizmos/ScalableSphereGizmo.h
@@ -6,7 +6,6 @@
 #include "InteractiveGizmo.h"
 #include "InteractiveGizmoBuilder.h"
 #include "TransformProxy.h"
-#include "GizmoActor.h"
 #include "BaseBehaviors/BehaviorTargetInterfaces.h"
 #include "BaseBehaviors/AnyButtonInputBehavior.h"
 
@@ -21,48 +20,9 @@ public:
 	virtual UInteractiveGizmo* BuildGizmo(const FToolBuilderState& SceneState) const override;
 };
 
-UCLASS()
-class INTERACTIVETOOLSFRAMEWORK_API AScalableSphereGizmoActor : public AGizmoActor
-{
-	GENERATED_BODY()
-
-	public:
-
-	AScalableSphereGizmoActor();
-
-	// +X Axis Handle
-	UPROPERTY()
-	UPrimitiveComponent* XPositive;
-
-	// -X Axis Handle
-	UPROPERTY()
-	UPrimitiveComponent* XNegative;
-
-	// +Y Axis Handle
-	UPROPERTY()
-	UPrimitiveComponent* YPositive;
-
-	// -Y Axis Handle
-	UPROPERTY()
-	UPrimitiveComponent* YNegative;
-
-	// +Z Axis Handle
-	UPROPERTY()
-	UPrimitiveComponent* ZPositive;
-
-	// -Z Axis Handle
-	UPROPERTY()
-	UPrimitiveComponent* ZNegative;
-};
-
 /**
- * UScalableSphereGizmo provides a sphere that can be scaled in all directions by dragging it's handles
- *
- * The in-scene representation of the Gizmo is a AScalableSphereGizmoActor (or subclass).
- * This Actor has FProperty members for the various sub-widgets, each as a separate Component.
- * Any particular sub-widget of the Gizmo can be disabled by setting the respective
- * Actor Component to null.
- *
+ * UScalableSphereGizmo provides a sphere that can be scaled in all directions by dragging
+ * anywhere on the three axial circles that represent it
  */
 UCLASS()
 class INTERACTIVETOOLSFRAMEWORK_API UScalableSphereGizmo : public UInteractiveGizmo, public IHoverBehaviorTarget
@@ -71,28 +31,20 @@ class INTERACTIVETOOLSFRAMEWORK_API UScalableSphereGizmo : public UInteractiveGi
 
 public:
 	// UInteractiveGizmo interface
-
 	virtual void Setup() override;
-
 	virtual void Render(IToolsContextRenderAPI* RenderAPI) override;
 
-	virtual void Shutdown() override;
-
 	// IHoverBehaviorTarget interface
-	// TODO: Add HoverBehavior
-	virtual FInputRayHit BeginHoverSequenceHitTest(const FInputDeviceRay& PressPos) override { return FInputRayHit(); }
-	virtual void OnBeginHover(const FInputDeviceRay& DevicePos) override {}
-	virtual bool OnUpdateHover(const FInputDeviceRay& DevicePos) override { return true; }
-	virtual void OnEndHover() override {}
+	virtual FInputRayHit BeginHoverSequenceHitTest(const FInputDeviceRay& PressPos) override;
+	virtual void OnBeginHover(const FInputDeviceRay& DevicePos) override { }
+	virtual bool OnUpdateHover(const FInputDeviceRay& DevicePos);
+	virtual void OnEndHover();
 	
 	/**
 	 * Set the Target to which the gizmo will be attached
 	 */
 	virtual void SetTarget(UTransformProxy* InTarget);
 
-	/* Set the World in which the GizmoActor will be spawned*/
-	virtual void SetWorld(UWorld* InWorld);
-
 	virtual void OnBeginDrag(const FInputDeviceRay& Ray);
 	virtual void OnUpdateDrag(const FInputDeviceRay& Ray);
 
@@ -104,35 +56,34 @@ public:
 
 	/** Called when the radius is chaged (by dragging or setting). Sends new radius as parameter. */
 	TFunction<void(const float)> UpdateRadiusFunc = nullptr;
-private:
 
-	/**
-	 * Create the GizmoActor and the handles
-	 */
-	void CreateGizmoHandles();
+	// The error threshold for hit detection with the sphere
+	UPROPERTY()
+	float HitErrorThreshold{ 12.f };
 
-	/**
-	 * Update the handles to the right position when the radius changes
-	 */
-	void UpdateGizmoHandles();
+	// The thickness of the sphere when being hovered over (in screen space)
+	UPROPERTY()
+	float HoverThickness{ 5.f };
 
-	/**
-	 * Callback for when the ActiveTarget's transform is changed
-	 */
-	void OnTransformChanged(UTransformProxy*, FTransform);
+private:
 
+	// Check if the Circle represented by CircleNormal centered at the gizmos location has been hit by the Ray
+	bool CheckCircleIntersection(const FRay& Ray, FVector CircleNormal, FVector& OutHitLocation, FVector& OutHitAxis);
+	
 	// The radius of the sphere
 	UPROPERTY()
 	float Radius;
 
+	// The thickness of the sphere
 	UPROPERTY()
-	UWorld* World;
+	float Thickness{ 0.f };
 
+	// Whether the sphere is currently being hovered over
 	UPROPERTY()
-	UTransformProxy* ActiveTarget;
+	bool bIsHovering{ false };
 
 	UPROPERTY()
-	AScalableSphereGizmoActor* GizmoActor;
+	UTransformProxy* ActiveTarget;
 
 	// The current axis that is being dragged along
 	UPROPERTY()
