Added hover behavior to all the new light gizmos

diff --git a/Engine/Plugins/Experimental/GizmoEdMode/Source/LightGizmos/Private/DirectionalLightGizmo.cpp b/Engine/Plugins/Experimental/GizmoEdMode/Source/LightGizmos/Private/DirectionalLightGizmo.cpp
--- a/Engine/Plugins/Experimental/GizmoEdMode/Source/LightGizmos/Private/DirectionalLightGizmo.cpp
+++ b/Engine/Plugins/Experimental/GizmoEdMode/Source/LightGizmos/Private/DirectionalLightGizmo.cpp
@@ -10,6 +10,7 @@
 #include "BaseGizmos/GizmoLineHandleComponent.h"
 #include "BaseGizmos/GizmoRenderingUtil.h"
 #include "SceneManagement.h"
+#include "BaseBehaviors/MouseHoverBehavior.h"
 
 // UDirectionalLightGizmoBuilder
 
@@ -50,6 +51,10 @@ void UDirectionalLightGizmo::Setup()
 	DirectionalLightBehavior->Initialize(this);
 	AddInputBehavior(DirectionalLightBehavior);
 
+	UMouseHoverBehavior* HoverBehavior = NewObject<UMouseHoverBehavior>(this);
+	HoverBehavior->Initialize(this);
+	AddInputBehavior(HoverBehavior);
+
 	CreateGizmoHandles();
 
 	CreateZRotationGizmo();
@@ -120,6 +125,56 @@ USubTransformProxy* UDirectionalLightGizmo::GetTransformProxy()
 	return TransformProxy;
 }
 
+FInputRayHit UDirectionalLightGizmo::BeginHoverSequenceHitTest(const FInputDeviceRay& PressPos)
+{
+	FHitResult HitResult;
+	FTransform DragTransform;
+
+	if (HitTest(PressPos.WorldRay, HitResult, DragTransform, HitComponent))
+	{
+		bIsHovering = true;
+		UpdateHandleColors();
+		return FInputRayHit(HitResult.Distance);
+	}
+
+	bIsHovering = false;
+	// Return invalid ray hit to say we don't want to listen to hover input
+	return FInputRayHit();
+}
+
+bool UDirectionalLightGizmo::OnUpdateHover(const FInputDeviceRay& DevicePos)
+{
+	if (!LightActor)
+	{
+		bIsHovering = false;
+		return false;
+	}
+
+	FVector Start = DevicePos.WorldRay.Origin;
+	const float MaxRaycastDistance = 1e6f;
+	FVector End = DevicePos.WorldRay.Origin + DevicePos.WorldRay.Direction * MaxRaycastDistance;
+
+	FRay HitCheckRay(Start, End - Start);
+	FHitResult HitResult;
+	FTransform DragTransform;
+
+	if (HitTest(HitCheckRay, HitResult, DragTransform, HitComponent))
+	{
+		bIsHovering = true;
+		return true;
+	}
+
+	bIsHovering = false;
+	return false;
+}
+
+void UDirectionalLightGizmo::OnEndHover()
+{
+	bIsHovering = false;
+
+	UpdateHandleColors();
+}
+
 void UDirectionalLightGizmo::SetSelectedObject(ADirectionalLight* InLight)
 {
 	LightActor = InLight;
@@ -160,16 +215,19 @@ void UDirectionalLightGizmo::OnBeginDrag(const FInputDeviceRay& Ray)
 	// Check if any component was hit
 	if (HitTest(HitCheckRay, HitResult, DragTransform, HitComponent))
 	{
+		bIsDragging = true;
+		UpdateHandleColors();
+
 		// Rotate around y axis if the arrow was hit
 		if (HitComponent == GizmoActor->Arrow)
 		{
-			HitAxis = FVector::YAxisVector;
+			HitAxis = LightActor->GetActorRotation().RotateVector(FVector::YAxisVector);
 
 			// Get the rotated plane vectors for the interaction
 			GizmoMath::MakeNormalPlaneBasis(HitAxis, RotationPlaneX, RotationPlaneZ);
 			RotationPlaneX = LightActor->GetActorRotation().RotateVector(FVector::XAxisVector);
 			RotationPlaneZ = LightActor->GetActorRotation().RotateVector(FVector::ZAxisVector);
-
+			
 		}
 		// Rotate around Z axis if the circle was hit
 		else
@@ -246,7 +304,14 @@ void UDirectionalLightGizmo::OnUpdateDrag(const FInputDeviceRay& Ray)
 
 }
 
-bool UDirectionalLightGizmo::HitTest(const FRay& Ray, FHitResult& OutHit, FTransform& OutTransform, UPrimitiveComponent*& OutHitComponent)
+void UDirectionalLightGizmo::OnEndDrag(const FInputDeviceRay& Ray)
+{
+	bIsDragging = false;
+
+	UpdateHandleColors();
+}
+
+bool UDirectionalLightGizmo::HitTest(const FRay& Ray, FHitResult& OutHit, FTransform& OutTransform, UGizmoBaseComponent*& OutHitComponent)
 {
 	FVector Start = Ray.Origin;
 	const float MaxRaycastDistance = 1e6f;
@@ -286,6 +351,27 @@ void UDirectionalLightGizmo::UpdateGizmoHandles()
 	}
 }
 
+void UDirectionalLightGizmo::UpdateHandleColors()
+{
+	if (bIsHovering || bIsDragging)
+	{
+		if (HitComponent)
+		{
+			HitComponent->Color = FLinearColor::Yellow;
+			HitComponent->NotifyExternalPropertyUpdates();
+		}
+	}
+	else
+	{
+		GizmoActor->Arrow->Color = FLinearColor::Red;
+		GizmoActor->Arrow->NotifyExternalPropertyUpdates();
+
+		GizmoActor->RotationZCircle->Color = FLinearColor::Blue;
+		GizmoActor->RotationZCircle->NotifyExternalPropertyUpdates();
+
+	}
+}
+
 void UDirectionalLightGizmo::OnTransformChanged(UTransformProxy*, FTransform)
 {
 	USceneComponent* GizmoComponent = GizmoActor->GetRootComponent();
@@ -330,7 +416,7 @@ FInputCaptureRequest UDirectionalLightGizmoInputBehavior::WantsCapture(const FIn
 	{
 		FHitResult HitResult;
 		FTransform DragTransform;
-		UPrimitiveComponent* HitComponent = nullptr;
+		UGizmoBaseComponent* HitComponent = nullptr;
 		if (Gizmo->HitTest(input.Mouse.WorldRay, HitResult, DragTransform, HitComponent))
 		{
 			return FInputCaptureRequest::Begin(this, EInputCaptureSide::Any, HitResult.Distance);
@@ -360,6 +446,7 @@ FInputCaptureUpdate UDirectionalLightGizmoInputBehavior::UpdateCapture(const FIn
 	if (IsReleased(input))
 	{
 		bInputDragCaptured = false;
+		Gizmo->OnEndDrag(LastWorldRay);
 		return FInputCaptureUpdate::End();
 	}
 
@@ -372,6 +459,7 @@ void UDirectionalLightGizmoInputBehavior::ForceEndCapture(const FInputCaptureDat
 {
 	if (bInputDragCaptured)
 	{
+		Gizmo->OnEndDrag(LastWorldRay);
 		bInputDragCaptured = false;
 	}
 }
\ No newline at end of file
diff --git a/Engine/Plugins/Experimental/GizmoEdMode/Source/LightGizmos/Private/DirectionalLightGizmo.h b/Engine/Plugins/Experimental/GizmoEdMode/Source/LightGizmos/Private/DirectionalLightGizmo.h
--- a/Engine/Plugins/Experimental/GizmoEdMode/Source/LightGizmos/Private/DirectionalLightGizmo.h
+++ b/Engine/Plugins/Experimental/GizmoEdMode/Source/LightGizmos/Private/DirectionalLightGizmo.h
@@ -12,6 +12,7 @@
 #include "Engine/DirectionalLight.h"
 #include "BaseGizmos/TransformGizmo.h"
 #include "SubTransformProxy.h"
+#include "BaseGizmos/GizmoBaseComponent.h"
 
 #include "DirectionalLightGizmo.generated.h"
 
@@ -34,10 +35,10 @@ public:
 	ADirectionalLightGizmoActor();
 
 	// The handle to rotate around its y axis
-	UPrimitiveComponent* Arrow;
+	UGizmoBaseComponent* Arrow;
 
 	// The handle to rotate around the world z axis
-	UPrimitiveComponent* RotationZCircle;
+	UGizmoBaseComponent* RotationZCircle;
 };
 
 /**
@@ -64,25 +65,28 @@ public:
 	USubTransformProxy* GetTransformProxy();
 
 	// IHoverBehaviorTarget interface
-	virtual FInputRayHit BeginHoverSequenceHitTest(const FInputDeviceRay& PressPos) override { return FInputRayHit(); }
+	virtual FInputRayHit BeginHoverSequenceHitTest(const FInputDeviceRay& PressPos) override;
 	virtual void OnBeginHover(const FInputDeviceRay& DevicePos) override {}
-	virtual bool OnUpdateHover(const FInputDeviceRay& DevicePos) override { return true; }
-	virtual void OnEndHover() override {}
+	virtual bool OnUpdateHover(const FInputDeviceRay& DevicePos) override;
+	virtual void OnEndHover() override;
 
 	virtual void SetSelectedObject(ADirectionalLight* InLight);
 	virtual void SetWorld(UWorld* InWorld);
 
 	virtual void OnBeginDrag(const FInputDeviceRay& Ray);
 	virtual void OnUpdateDrag(const FInputDeviceRay& Ray);
+	virtual void OnEndDrag(const FInputDeviceRay& Ray);
 
 	/** Check if any of the components are hit by the input ray */
-	bool HitTest(const FRay& Ray, FHitResult& OutHit, FTransform& OutTransform, UPrimitiveComponent*& OutHitComponent);
+	bool HitTest(const FRay& Ray, FHitResult& OutHit, FTransform& OutTransform, UGizmoBaseComponent*& OutHitComponent);
 
 private:
 
 	void CreateGizmoHandles();
 	void UpdateGizmoHandles();
 
+	void UpdateHandleColors();
+
 	void OnTransformChanged(UTransformProxy*, FTransform);
 
 	/** Create the gizmo that rotates around the world z axis */
@@ -103,6 +107,12 @@ private:
 	UPROPERTY()
 	ADirectionalLight* LightActor;
 
+	UPROPERTY()
+	bool bIsHovering{ false };
+
+	UPROPERTY()
+	bool bIsDragging{ false };
+
 	/** Parameters used during hit testing */
 	UPROPERTY()
 	FVector DragStartWorldPosition;
@@ -123,7 +133,7 @@ private:
 	FVector RotationPlaneZ;
 
 	UPROPERTY()
-	UPrimitiveComponent* HitComponent;
+	UGizmoBaseComponent* HitComponent;
 
 	UPROPERTY()
 	float ArrowLength{ 120.f };
diff --git a/Engine/Plugins/Experimental/GizmoEdMode/Source/LightGizmos/Private/ScalableConeGizmo.cpp b/Engine/Plugins/Experimental/GizmoEdMode/Source/LightGizmos/Private/ScalableConeGizmo.cpp
--- a/Engine/Plugins/Experimental/GizmoEdMode/Source/LightGizmos/Private/ScalableConeGizmo.cpp
+++ b/Engine/Plugins/Experimental/GizmoEdMode/Source/LightGizmos/Private/ScalableConeGizmo.cpp
@@ -42,7 +42,8 @@ void UScalableConeGizmo::Render(IToolsContextRenderAPI* RenderAPI)
 	{
 		DrawWireSphereCappedCone(RenderAPI->GetPrimitiveDrawInterface(), ActiveTarget->GetTransform(), Length, Angle, 32, 8, 10, ConeColor, SDPG_Foreground);
 
-		if (bIsHovering)
+		// Draw the yellow circle as a hightlight when hovering or dragging
+		if (bIsHovering || bIsDragging)
 		{
 			FVector WorldOrigin = ActiveTarget->GetTransform().GetLocation();
 			FVector CircleNormal = ActiveTarget->GetTransform().GetRotation().Vector();
@@ -59,24 +60,19 @@ void UScalableConeGizmo::Render(IToolsContextRenderAPI* RenderAPI)
 			// Direction vectors to represent the circle
 			GizmoMath::MakeNormalPlaneBasis(CircleNormal, CircleX, CircleY);
 
-			float CircleThickness = HoverThickness;
-
-			// Figure out the correct thickness of the Circle in world coordinates
-			const FSceneView* View = RenderAPI->GetSceneView();
-
-			FVector CircleEndWithThickness = DragStartWorldPosition + CircleX * CircleThickness;
-			float ScreenThickness = GizmoRenderingUtil::CalculateLocalPixelToWorldScale(View, CircleEndWithThickness);
-
-			CircleThickness *= ScreenThickness;
-
-			DrawCircle(RenderAPI->GetPrimitiveDrawInterface(), CircleOrigin, CircleX, CircleY, ConeColor, ConeRadius, 32, SDPG_Foreground, CircleThickness);
+			DrawCircle(RenderAPI->GetPrimitiveDrawInterface(), CircleOrigin, CircleX, CircleY, FLinearColor::Yellow, ConeRadius, 32, SDPG_Foreground, 0, 0.1f);
+		}
 
+		// Draw the lines to show which direction to drag in while hovering
+		if (bIsHovering)
+		{
 			// Parameters for the line that shows the drag direction
-			FVector LineStart = DragStartWorldPosition;
+			FVector LineStart = DragCurrentPositionProjected;
 			float LineLength = 30.f;
 			FVector LineEnd = LineStart + HitAxis * LineLength;
 
 			// Get the Pixel to World scale of the line
+			const FSceneView* View = RenderAPI->GetSceneView();
 			float PixelToWorld = GizmoRenderingUtil::CalculateLocalPixelToWorldScale(View, LineEnd);
 
 			// Draw the lines in both directions
@@ -128,11 +124,13 @@ FInputRayHit UScalableConeGizmo::BeginHoverSequenceHitTest(const FInputDeviceRay
 	{
 		bIsHovering = true;
 		DragStartWorldPosition = DragTransform.GetLocation();
+		DragCurrentPositionProjected = DragStartWorldPosition;
 		HitAxis = TestHitAxis;
 
 		return FInputRayHit(HitResult.Distance);
 	}
 
+	bIsHovering = false;
 	// Return invalid ray hit to say we don't want to listen to hover input
 	return FInputRayHit();
 }
@@ -141,6 +139,7 @@ bool UScalableConeGizmo::OnUpdateHover(const FInputDeviceRay& DevicePos)
 {
 	if (!ActiveTarget)
 	{
+		bIsHovering = false;
 		return false;
 	}
 
@@ -157,10 +156,12 @@ bool UScalableConeGizmo::OnUpdateHover(const FInputDeviceRay& DevicePos)
 	{
 		bIsHovering = true;
 		DragStartWorldPosition = DragTransform.GetLocation();
+		DragCurrentPositionProjected = DragStartWorldPosition;
 		HitAxis = TestHitAxis;
 		return true;
 	}
 
+	bIsHovering = false;
 	return false;
 }
 
@@ -251,6 +252,9 @@ void UScalableConeGizmo::OnBeginDrag(const FInputDeviceRay& Ray)
 			RayNearestPt, RayNearestParam);
 
 		DragStartWorldPosition = DragTransform.GetLocation();
+		DragCurrentPositionProjected = DragStartWorldPosition;
+
+		bIsDragging = true;
 	}
 }
 
@@ -294,9 +298,15 @@ void UScalableConeGizmo::OnUpdateDrag(const FInputDeviceRay& Ray)
 	SetAngleDegrees(Angle + FMath::RadiansToDegrees(DeltaAngle));
 
 	InteractionStartPoint = AxisNearestPt;
+	DragCurrentPositionProjected = AxisNearestPt;
 	InteractionStartParameter = AxisNearestParam;
 }
 
+void UScalableConeGizmo::OnEndDrag(const FInputDeviceRay& Ray)
+{
+	bIsDragging = false;
+}
+
 // UScalableConeGizmoInputBehavior
 
 void UScalableConeGizmoInputBehavior::Initialize(UScalableConeGizmo* InGizmo)
@@ -340,6 +350,7 @@ FInputCaptureUpdate UScalableConeGizmoInputBehavior::UpdateCapture(const FInputD
 	if (IsReleased(input))
 	{
 		bInputDragCaptured = false;
+		Gizmo->OnEndDrag(LastWorldRay);
 		return FInputCaptureUpdate::End();
 	}
 
@@ -353,5 +364,6 @@ void UScalableConeGizmoInputBehavior::ForceEndCapture(const FInputCaptureData& d
 	if (bInputDragCaptured)
 	{
 		bInputDragCaptured = false;
+		Gizmo->OnEndDrag(LastWorldRay);
 	}
 }
\ No newline at end of file
diff --git a/Engine/Plugins/Experimental/GizmoEdMode/Source/LightGizmos/Private/SpotLightGizmo.cpp b/Engine/Plugins/Experimental/GizmoEdMode/Source/LightGizmos/Private/SpotLightGizmo.cpp
--- a/Engine/Plugins/Experimental/GizmoEdMode/Source/LightGizmos/Private/SpotLightGizmo.cpp
+++ b/Engine/Plugins/Experimental/GizmoEdMode/Source/LightGizmos/Private/SpotLightGizmo.cpp
@@ -9,6 +9,7 @@
 #include "BaseGizmos/GizmoMath.h"
 #include "Components/SphereComponent.h"
 #include "BaseGizmos/GizmoLineHandleComponent.h"
+#include "BaseBehaviors/MouseHoverBehavior.h"
 
 // USpotLightGizmoBuilder
 
@@ -45,6 +46,10 @@ void USpotLightGizmo::Setup()
 	USpotLightGizmoInputBehavior* SpotLightBehavior = NewObject<USpotLightGizmoInputBehavior>(this);
 	SpotLightBehavior->Initialize(this);
 	AddInputBehavior(SpotLightBehavior);
+
+	UMouseHoverBehavior* HoverBehavior = NewObject<UMouseHoverBehavior>(this);
+	HoverBehavior->Initialize(this);
+	AddInputBehavior(HoverBehavior);
 }
 
 void USpotLightGizmo::Tick(float DeltaTime)
@@ -89,6 +94,80 @@ void USpotLightGizmo::Shutdown()
 	}
 }
 
+FInputRayHit USpotLightGizmo::BeginHoverSequenceHitTest(const FInputDeviceRay& PressPos)
+{
+	FHitResult HitResult;
+	FVector HitAxis;
+	FTransform DragTransform;
+
+	if (HitTest(PressPos.WorldRay, HitResult, DragTransform))
+	{
+		bIsHovering = true;
+		DragStartWorldPosition = DragTransform.GetLocation();
+		UpdateHandleColors();
+		return FInputRayHit(HitResult.Distance);
+	}
+
+	// check to avoid refreshing handle colors constantly even if not required
+	if (bIsHovering)
+	{
+		bIsHovering = false;
+		UpdateHandleColors();
+	}
+	
+	// Return invalid ray hit to say we don't want to listen to hover input
+	return FInputRayHit();
+}
+
+bool USpotLightGizmo::OnUpdateHover(const FInputDeviceRay& DevicePos)
+{
+	if (!LightActor)
+	{
+		bIsHovering = false;
+		return false;
+	}
+
+	FVector Start = DevicePos.WorldRay.Origin;
+	const float MaxRaycastDistance = 1e6f;
+	FVector End = DevicePos.WorldRay.Origin + DevicePos.WorldRay.Direction * MaxRaycastDistance;
+
+	FRay HitCheckRay(Start, End - Start);
+	FHitResult HitResult;
+	FTransform DragTransform;
+
+	if (HitTest(HitCheckRay, HitResult, DragTransform))
+	{
+		bIsHovering = true;
+		DragStartWorldPosition = DragTransform.GetLocation();
+		return true;
+	}
+
+	bIsHovering = false;
+	return false;
+}
+
+void USpotLightGizmo::OnEndHover()
+{
+	bIsHovering = false;
+
+	UpdateHandleColors();
+}
+
+void USpotLightGizmo::UpdateHandleColors()
+{
+	if (bIsHovering || bIsDragging)
+	{
+		GizmoActor->AttenuationScaleHandle->Color = FLinearColor::Yellow;
+	}
+	else
+	{
+		GizmoActor->AttenuationScaleHandle->Color = FLinearColor::Blue;
+	}
+
+	GizmoActor->AttenuationScaleHandle->NotifyExternalPropertyUpdates();
+
+}
+
 void USpotLightGizmo::SetSelectedObject(ASpotLight* InLight)
 {
 	LightActor = InLight;
@@ -205,6 +284,10 @@ void USpotLightGizmo::OnBeginDrag(const FInputDeviceRay& Ray)
 			RayNearestPt, RayNearestParam);
 
 		DragStartWorldPosition = DragTransform.GetLocation();
+
+		bIsDragging = true;
+
+		UpdateHandleColors();
 	}
 }
 
@@ -244,6 +327,13 @@ void USpotLightGizmo::OnUpdateDrag(const FInputDeviceRay& Ray)
 	}
 }
 
+void USpotLightGizmo::OnEndDrag(const FInputDeviceRay& Ray)
+{
+	bIsDragging = false;
+
+	UpdateHandleColors();
+}
+
 bool USpotLightGizmo::HitTest(const FRay& Ray, FHitResult& OutHit, FTransform& OutTransform)
 {
 	if (!GizmoActor)
@@ -327,6 +417,7 @@ FInputCaptureUpdate USpotLightGizmoInputBehavior::UpdateCapture(const FInputDevi
 	if (IsReleased(input))
 	{
 		bInputDragCaptured = false;
+		Gizmo->OnEndDrag(LastWorldRay);
 		return FInputCaptureUpdate::End();
 	}
 
@@ -339,6 +430,7 @@ void USpotLightGizmoInputBehavior::ForceEndCapture(const FInputCaptureData& data
 {
 	if (bInputDragCaptured)
 	{
+		Gizmo->OnEndDrag(LastWorldRay);
 		bInputDragCaptured = false;
 	}
 }
\ No newline at end of file
diff --git a/Engine/Plugins/Experimental/GizmoEdMode/Source/LightGizmos/Private/SpotLightGizmo.h b/Engine/Plugins/Experimental/GizmoEdMode/Source/LightGizmos/Private/SpotLightGizmo.h
--- a/Engine/Plugins/Experimental/GizmoEdMode/Source/LightGizmos/Private/SpotLightGizmo.h
+++ b/Engine/Plugins/Experimental/GizmoEdMode/Source/LightGizmos/Private/SpotLightGizmo.h
@@ -6,6 +6,8 @@
 #include "Engine/SpotLight.h"
 #include "InteractiveGizmo.h"
 #include "SubTransformProxy.h"
+#include "BaseBehaviors/BehaviorTargetInterfaces.h"
+#include "BaseGizmos/GizmoBaseComponent.h"
 
 #include "SpotLightGizmo.generated.h"
 
@@ -28,7 +30,7 @@ public:
 	ASpotLightGizmoActor();
 
 	 // The handle to drag and scale the attenuation
-	UPrimitiveComponent* AttenuationScaleHandle;
+	UGizmoBaseComponent* AttenuationScaleHandle;
 };
 
 
@@ -38,7 +40,7 @@ public:
  *
  */
 UCLASS()
-class USpotLightGizmo : public UInteractiveGizmo
+class USpotLightGizmo : public UInteractiveGizmo, public IHoverBehaviorTarget
 {
 	GENERATED_BODY()
 
@@ -51,6 +53,12 @@ public:
 
 	virtual void Shutdown() override;
 
+	// IHoverBehaviorTarget interface
+	virtual FInputRayHit BeginHoverSequenceHitTest(const FInputDeviceRay& PressPos) override;
+	virtual void OnBeginHover(const FInputDeviceRay& DevicePos) override {}
+	virtual bool OnUpdateHover(const FInputDeviceRay& DevicePos) override;
+	virtual void OnEndHover() override;
+
 	USpotLightGizmo();
 
 	void SetSelectedObject(ASpotLight* InLight);
@@ -70,6 +78,7 @@ public:
 
 	virtual void OnBeginDrag(const FInputDeviceRay& Ray);
 	virtual void OnUpdateDrag(const FInputDeviceRay& Ray);
+	virtual void OnEndDrag(const FInputDeviceRay& Ray);
 
 	/** Check if the Input ray hit any of the components of the gizmo*/
 	bool HitTest(const FRay& Ray, FHitResult& OutHit, FTransform& OutTransform);
@@ -109,9 +118,17 @@ private:
 	UPROPERTY()
 	float InteractionStartParameter;
 
+	UPROPERTY()
+	bool bIsHovering{ false };
+
+	UPROPERTY()
+	bool bIsDragging{ false };
+
 	void OnOuterAngleUpdate(float NewAngle);
 	void OnInnerAngleUpdate(float NewAngle);
 	void OnTransformChanged(UTransformProxy*, FTransform);
+
+	void UpdateHandleColors();
 };
 
 /**
diff --git a/Engine/Plugins/Experimental/GizmoEdMode/Source/LightGizmos/Public/ScalableConeGizmo.h b/Engine/Plugins/Experimental/GizmoEdMode/Source/LightGizmos/Public/ScalableConeGizmo.h
--- a/Engine/Plugins/Experimental/GizmoEdMode/Source/LightGizmos/Public/ScalableConeGizmo.h
+++ b/Engine/Plugins/Experimental/GizmoEdMode/Source/LightGizmos/Public/ScalableConeGizmo.h
@@ -54,6 +54,7 @@ public:
 
 	virtual void OnBeginDrag(const FInputDeviceRay& Ray);
 	virtual void OnUpdateDrag(const FInputDeviceRay& Ray);
+	virtual void OnEndDrag(const FInputDeviceRay& Ray);
 
 	/** Check if the input ray hits the cone. The hit testable parts of the cone currently is just 
 	 *  the circle at its base
@@ -76,9 +77,6 @@ public:
 	UPROPERTY()
 	float HitErrorThreshold{ 12.f };
 
-	// The thickness of the base of the cone when being hovered over (in screen space)
-	UPROPERTY()
-	float HoverThickness{ 5.f };
 
 	/** Called when the Angle of the cone is changed. Sends new angle as parameter. */
 	TFunction<void(const float)> UpdateAngleFunc = nullptr;
@@ -96,6 +94,10 @@ private:
 	UPROPERTY()
 	bool bIsHovering{ false };
 
+	/** Whether the gizmo is being dragged */
+	UPROPERTY()
+	bool bIsDragging{ false };
+
 	UPROPERTY()
 	UTransformProxy* ActiveTarget;
 
@@ -104,6 +106,10 @@ private:
 	UPROPERTY()
 	FVector DragStartWorldPosition;
 
+	// The position the drag is on currently (projected onto the line it is being dragged along)
+	UPROPERTY()
+	FVector DragCurrentPositionProjected;
+
 	UPROPERTY()
 	FVector InteractionStartPoint;
 
diff --git a/Engine/Source/Runtime/Experimental/InteractiveToolsFramework/Private/BaseGizmos/ScalableSphereGizmo.cpp b/Engine/Source/Runtime/Experimental/InteractiveToolsFramework/Private/BaseGizmos/ScalableSphereGizmo.cpp
--- a/Engine/Source/Runtime/Experimental/InteractiveToolsFramework/Private/BaseGizmos/ScalableSphereGizmo.cpp
+++ b/Engine/Source/Runtime/Experimental/InteractiveToolsFramework/Private/BaseGizmos/ScalableSphereGizmo.cpp
@@ -40,14 +40,15 @@ void UScalableSphereGizmo::Render(IToolsContextRenderAPI* RenderAPI)
 {
 	if (ActiveTarget)
 	{
-		float UseThickness = Thickness;
+		FColor UseColor = FColor(200, 255, 255);
 
-		if (bIsHovering)
+		// Draw the gizmo in yellow and draw lines showing drag direction if hovering or dragging
+		if (bIsHovering || bIsDragging)
 		{
-			UseThickness = HoverThickness;
+			UseColor = FColor::Yellow;
 
 			// Parameters for the line that shows the drag direction
-			FVector LineStart = DragStartWorldPosition;
+			FVector LineStart = DragCurrentPositionProjected;
 			float LineLength = 30.f;
 			FVector LineEnd = LineStart + ActiveAxis * LineLength;
 
@@ -57,16 +58,10 @@ void UScalableSphereGizmo::Render(IToolsContextRenderAPI* RenderAPI)
 
 			// Draw the lines in both directions
 			RenderAPI->GetPrimitiveDrawInterface()->DrawLine(LineStart, LineStart + ActiveAxis * PixelToWorld * LineLength, FLinearColor::Red, SDPG_Foreground);
-			RenderAPI->GetPrimitiveDrawInterface()->DrawLine(LineStart, LineStart - ActiveAxis * PixelToWorld* LineLength, FLinearColor::Red, SDPG_Foreground);
-
-			// Figure out the correct thickness of the sphere in world coordinates
-			FVector CircleEndWithThickness = LineStart + ActiveAxis * UseThickness;
-			float ScreenThickness = GizmoRenderingUtil::CalculateLocalPixelToWorldScale(View, CircleEndWithThickness);
-
-			UseThickness *= ScreenThickness;
+			RenderAPI->GetPrimitiveDrawInterface()->DrawLine(LineStart, LineStart - ActiveAxis * PixelToWorld * LineLength, FLinearColor::Red, SDPG_Foreground);
 		}
 
-		DrawWireSphereAutoSides(RenderAPI->GetPrimitiveDrawInterface(), ActiveTarget->GetTransform(), FColor(200, 255, 255), Radius, SDPG_Foreground, UseThickness);
+		DrawWireSphereAutoSides(RenderAPI->GetPrimitiveDrawInterface(), ActiveTarget->GetTransform(), UseColor, Radius, SDPG_Foreground);
 	}
 }
 
@@ -80,11 +75,13 @@ FInputRayHit UScalableSphereGizmo::BeginHoverSequenceHitTest(const FInputDeviceR
 	{
 		bIsHovering = true;
 		DragStartWorldPosition = DragTransform.GetLocation();
+		DragCurrentPositionProjected = DragStartWorldPosition;
 		ActiveAxis = HitAxis;
 		
 		return FInputRayHit(HitResult.Distance);
 	}
 	
+	bIsHovering = false;
 	// Return invalid ray hit to say we don't want to listen to hover input
 	return FInputRayHit();
 }
@@ -93,6 +90,7 @@ bool UScalableSphereGizmo::OnUpdateHover(const FInputDeviceRay& DevicePos)
 {
 	if (!ActiveTarget)
 	{
+		bIsHovering = false;
 		return false;
 	}
 
@@ -109,10 +107,12 @@ bool UScalableSphereGizmo::OnUpdateHover(const FInputDeviceRay& DevicePos)
 	{
 		bIsHovering = true;
 		DragStartWorldPosition = DragTransform.GetLocation();
+		DragCurrentPositionProjected = DragStartWorldPosition;
 		ActiveAxis = HitAxis;
 		return true;
 	}
 
+	bIsHovering = false;
 	return false;
 }
 
@@ -251,6 +251,9 @@ void UScalableSphereGizmo::OnBeginDrag(const FInputDeviceRay& Ray)
 			RayNearestPt, RayNearestParam);
 
 		DragStartWorldPosition = DragTransform.GetLocation();
+		DragCurrentPositionProjected = DragStartWorldPosition;
+
+		bIsDragging = true;
 	}
 }
 
@@ -273,10 +276,17 @@ void UScalableSphereGizmo::OnUpdateDrag(const FInputDeviceRay& Ray)
 
 	InteractionStartParameter = InteractionCurrentParameter;
 
+	DragCurrentPositionProjected = AxisNearestPt;
+
 	 // Update the radius
 	SetRadius(Radius + DeltaParam);
 }
 
+void UScalableSphereGizmo::OnEndDrag(const FInputDeviceRay& Ray)
+{
+	bIsDragging = false;
+}
+
 // UScalableSphereGizmoInputBehavior
 
 void UScalableSphereGizmoInputBehavior::Initialize(UScalableSphereGizmo* InGizmo)
@@ -322,6 +332,7 @@ FInputCaptureUpdate UScalableSphereGizmoInputBehavior::UpdateCapture(const FInpu
 	if (IsReleased(input))
 	{
 		bInputDragCaptured = false;
+		Gizmo->OnEndDrag(LastWorldRay);
 		return FInputCaptureUpdate::End();
 	}
 
@@ -336,5 +347,6 @@ void UScalableSphereGizmoInputBehavior::ForceEndCapture(const FInputCaptureData&
 	if (bInputDragCaptured)
 	{
 		bInputDragCaptured = false;
+		Gizmo->OnEndDrag(LastWorldRay);
 	}
 }
\ No newline at end of file
diff --git a/Engine/Source/Runtime/Experimental/InteractiveToolsFramework/Public/BaseGizmos/ScalableSphereGizmo.h b/Engine/Source/Runtime/Experimental/InteractiveToolsFramework/Public/BaseGizmos/ScalableSphereGizmo.h
--- a/Engine/Source/Runtime/Experimental/InteractiveToolsFramework/Public/BaseGizmos/ScalableSphereGizmo.h
+++ b/Engine/Source/Runtime/Experimental/InteractiveToolsFramework/Public/BaseGizmos/ScalableSphereGizmo.h
@@ -47,6 +47,7 @@ public:
 
 	virtual void OnBeginDrag(const FInputDeviceRay& Ray);
 	virtual void OnUpdateDrag(const FInputDeviceRay& Ray);
+	virtual void OnEndDrag(const FInputDeviceRay& Ray);
 
 	/** Check if the input Ray hit any of the components of the internal actor */
 	bool HitTest(const FRay& Ray, FHitResult& OutHit, FVector& OutAxis, FTransform& OutTransform);
@@ -61,10 +62,6 @@ public:
 	UPROPERTY()
 	float HitErrorThreshold{ 12.f };
 
-	// The thickness of the sphere when being hovered over (in screen space)
-	UPROPERTY()
-	float HoverThickness{ 5.f };
-
 private:
 
 	// Check if the Circle represented by CircleNormal centered at the gizmos location has been hit by the Ray
@@ -74,14 +71,14 @@ private:
 	UPROPERTY()
 	float Radius;
 
-	// The thickness of the sphere
-	UPROPERTY()
-	float Thickness{ 0.f };
-
 	// Whether the sphere is currently being hovered over
 	UPROPERTY()
 	bool bIsHovering{ false };
 
+	// Whether the sphere is currently being dragged
+	UPROPERTY()
+	bool bIsDragging{ false };
+
 	UPROPERTY()
 	UTransformProxy* ActiveTarget;
 
@@ -93,6 +90,10 @@ private:
 	UPROPERTY()
 	FVector DragStartWorldPosition;
 
+	// The position the drag is on currently (projected onto the line it is being dragged along)
+	UPROPERTY()
+	FVector DragCurrentPositionProjected;
+
 	// The initial parameter along the drag axist
 	UPROPERTY()
 	float InteractionStartParameter;
