Added support for undo/redo in all the light gizmo operations and fixed a crash caused by undo-ing twice in a row on the light gizmos.

diff --git a/Engine/Plugins/Experimental/GizmoEdMode/Source/LightGizmos/Private/DirectionalLightGizmo.cpp b/Engine/Plugins/Experimental/GizmoEdMode/Source/LightGizmos/Private/DirectionalLightGizmo.cpp
--- a/Engine/Plugins/Experimental/GizmoEdMode/Source/LightGizmos/Private/DirectionalLightGizmo.cpp
+++ b/Engine/Plugins/Experimental/GizmoEdMode/Source/LightGizmos/Private/DirectionalLightGizmo.cpp
@@ -12,6 +12,9 @@
 #include "SceneManagement.h"
 #include "BaseBehaviors/MouseHoverBehavior.h"
 
+#define LOCTEXT_NAMESPACE "UDirectionalLightGizmo"
+
+
 // UDirectionalLightGizmoBuilder
 
 UInteractiveGizmo* UDirectionalLightGizmoBuilder::BuildGizmo(const FToolBuilderState& SceneState) const
@@ -227,6 +230,8 @@ void UDirectionalLightGizmo::OnBeginDrag(const FInputDeviceRay& Ray)
 			GizmoMath::MakeNormalPlaneBasis(HitAxis, RotationPlaneX, RotationPlaneZ);
 			RotationPlaneX = LightActor->GetActorRotation().RotateVector(FVector::XAxisVector);
 			RotationPlaneZ = LightActor->GetActorRotation().RotateVector(FVector::ZAxisVector);
+
+			GetGizmoManager()->BeginUndoTransaction(LOCTEXT("DirectionalLightYRotation", "Directional Light Y Rotation"));
 			
 		}
 		// Rotate around Z axis if the circle was hit
@@ -235,6 +240,8 @@ void UDirectionalLightGizmo::OnBeginDrag(const FInputDeviceRay& Ray)
 			HitAxis = FVector::ZAxisVector;
 			RotationPlaneX = FVector::XAxisVector;
 			RotationPlaneZ = FVector::YAxisVector;
+
+			GetGizmoManager()->BeginUndoTransaction(LOCTEXT("DirectionalLightZRotation", "Directional Light Z Rotation"));
 		}
 
 		// Calculate initial hit position
@@ -281,6 +288,8 @@ void UDirectionalLightGizmo::OnUpdateDrag(const FInputDeviceRay& Ray)
 
 	float DeltaAngle = InteractionCurAngle - InteractionStartParameter;
 
+	LightActor->Modify();
+
 	// Rotate around y axis if the arrow was hit
 	if (HitComponent == GizmoActor->Arrow)
 	{
@@ -309,6 +318,8 @@ void UDirectionalLightGizmo::OnEndDrag(const FInputDeviceRay& Ray)
 	bIsDragging = false;
 
 	UpdateHandleColors();
+
+	GetGizmoManager()->EndUndoTransaction();
 }
 
 bool UDirectionalLightGizmo::HitTest(const FRay& Ray, FHitResult& OutHit, FTransform& OutTransform, UGizmoBaseComponent*& OutHitComponent)
@@ -374,6 +385,11 @@ void UDirectionalLightGizmo::UpdateHandleColors()
 
 void UDirectionalLightGizmo::OnTransformChanged(UTransformProxy*, FTransform)
 {
+	if (!GizmoActor)
+	{
+		return;
+	}
+
 	USceneComponent* GizmoComponent = GizmoActor->GetRootComponent();
 
 	FTransform TargetTransform = TransformProxy->GetTransform();
@@ -462,4 +478,6 @@ void UDirectionalLightGizmoInputBehavior::ForceEndCapture(const FInputCaptureDat
 		Gizmo->OnEndDrag(LastWorldRay);
 		bInputDragCaptured = false;
 	}
-}
\ No newline at end of file
+}
+
+#undef LOCTEXT_NAMESPACE
\ No newline at end of file
diff --git a/Engine/Plugins/Experimental/GizmoEdMode/Source/LightGizmos/Private/PointLightGizmo.cpp b/Engine/Plugins/Experimental/GizmoEdMode/Source/LightGizmos/Private/PointLightGizmo.cpp
--- a/Engine/Plugins/Experimental/GizmoEdMode/Source/LightGizmos/Private/PointLightGizmo.cpp
+++ b/Engine/Plugins/Experimental/GizmoEdMode/Source/LightGizmos/Private/PointLightGizmo.cpp
@@ -3,6 +3,8 @@
 #include "PointLightGizmo.h"
 #include "Components/PointLightComponent.h"
 
+#define LOCTEXT_NAMESPACE "UScalableSphereGizmo"
+
 UInteractiveGizmo* UPointLightGizmoBuilder::BuildGizmo(const FToolBuilderState& SceneState) const
 {
 	UPointLightGizmo* NewGizmo = NewObject<UPointLightGizmo>(SceneState.GizmoManager);
@@ -75,6 +77,8 @@ void UPointLightGizmo::CreateLightGizmo()
 	AttenuationGizmo->SetTarget(TransformProxy);
 
 	AttenuationGizmo->UpdateRadiusFunc = [this](float NewRadius) { this->OnAttenuationUpdate(NewRadius); };
+
+	AttenuationGizmo->TransactionDescription = LOCTEXT("PointLightGizmo", "Point Light Attenuation");
 }
 
 void UPointLightGizmo::SetWorld(UWorld* InWorld)
@@ -86,7 +90,10 @@ void UPointLightGizmo::OnAttenuationUpdate(float NewRadius)
 {
 	// Update the attenuation radius
 	// TODO: Is this the right way to update radius?
+
+	LightActor->PointLightComponent->Modify();
 	LightActor->PointLightComponent->AttenuationRadius = NewRadius;
-	LightActor->PointLightComponent->MarkRenderStateDirty();
+
 }
 
+#undef LOCTEXT_NAMESPACE
\ No newline at end of file
diff --git a/Engine/Plugins/Experimental/GizmoEdMode/Source/LightGizmos/Private/ScalableConeGizmo.cpp b/Engine/Plugins/Experimental/GizmoEdMode/Source/LightGizmos/Private/ScalableConeGizmo.cpp
--- a/Engine/Plugins/Experimental/GizmoEdMode/Source/LightGizmos/Private/ScalableConeGizmo.cpp
+++ b/Engine/Plugins/Experimental/GizmoEdMode/Source/LightGizmos/Private/ScalableConeGizmo.cpp
@@ -8,6 +8,8 @@
 #include "BaseGizmos/GizmoRenderingUtil.h"
 #include "BaseBehaviors/MouseHoverBehavior.h"
 
+#define LOCTEXT_NAMESPACE "UScalableSphereGizmo"
+
 // UScalableConeGizmoBuilder
 
 UInteractiveGizmo* UScalableConeGizmoBuilder::BuildGizmo(const FToolBuilderState& SceneState) const
@@ -27,6 +29,8 @@ void UScalableConeGizmo::Setup()
 	MinAngle = 0.f;
 	ConeColor = FColor(200, 255, 255);
 
+	TransactionDescription = LOCTEXT("ScalableConeGizmo", "Scale Cone Gizmo");
+
 	UScalableConeGizmoInputBehavior* ScalableConeBehavior = NewObject<UScalableConeGizmoInputBehavior>(this);
 	ScalableConeBehavior->Initialize(this);
 	AddInputBehavior(ScalableConeBehavior);
@@ -255,6 +259,8 @@ void UScalableConeGizmo::OnBeginDrag(const FInputDeviceRay& Ray)
 		DragCurrentPositionProjected = DragStartWorldPosition;
 
 		bIsDragging = true;
+
+		GetGizmoManager()->BeginUndoTransaction(TransactionDescription);
 	}
 }
 
@@ -304,6 +310,7 @@ void UScalableConeGizmo::OnUpdateDrag(const FInputDeviceRay& Ray)
 
 void UScalableConeGizmo::OnEndDrag(const FInputDeviceRay& Ray)
 {
+	GetGizmoManager()->EndUndoTransaction();
 	bIsDragging = false;
 }
 
@@ -366,4 +373,6 @@ void UScalableConeGizmoInputBehavior::ForceEndCapture(const FInputCaptureData& d
 		bInputDragCaptured = false;
 		Gizmo->OnEndDrag(LastWorldRay);
 	}
-}
\ No newline at end of file
+}
+
+#undef LOCTEXT_NAMESPACE
\ No newline at end of file
diff --git a/Engine/Plugins/Experimental/GizmoEdMode/Source/LightGizmos/Private/SpotLightGizmo.cpp b/Engine/Plugins/Experimental/GizmoEdMode/Source/LightGizmos/Private/SpotLightGizmo.cpp
--- a/Engine/Plugins/Experimental/GizmoEdMode/Source/LightGizmos/Private/SpotLightGizmo.cpp
+++ b/Engine/Plugins/Experimental/GizmoEdMode/Source/LightGizmos/Private/SpotLightGizmo.cpp
@@ -11,6 +11,8 @@
 #include "BaseGizmos/GizmoLineHandleComponent.h"
 #include "BaseBehaviors/MouseHoverBehavior.h"
 
+#define LOCTEXT_NAMESPACE "USpotLightGizmo"
+
 // USpotLightGizmoBuilder
 
 UInteractiveGizmo* USpotLightGizmoBuilder::BuildGizmo(const FToolBuilderState& SceneState) const
@@ -189,6 +191,8 @@ USubTransformProxy* USpotLightGizmo::GetTransformProxy()
 
 void USpotLightGizmo::OnOuterAngleUpdate(float NewAngle)
 {
+	LightActor->SpotLightComponent->Modify();
+
 	LightActor->SpotLightComponent->OuterConeAngle = NewAngle;
 	
 	// OuterAngle cannot be less than Inner Angle
@@ -203,6 +207,8 @@ void USpotLightGizmo::OnOuterAngleUpdate(float NewAngle)
 
 void USpotLightGizmo::OnInnerAngleUpdate(float NewAngle)
 {
+	LightActor->SpotLightComponent->Modify();
+
 	LightActor->SpotLightComponent->InnerConeAngle = NewAngle;
 
 	// Inner Angle cannot be greater than Outer Angle
@@ -216,6 +222,11 @@ void USpotLightGizmo::OnInnerAngleUpdate(float NewAngle)
 
 void USpotLightGizmo::OnTransformChanged(UTransformProxy*, FTransform)
 {
+	if (!GizmoActor)
+	{
+		return;
+	}
+
 	USceneComponent* GizmoComponent = GizmoActor->GetRootComponent();
 
 	FTransform TargetTransform = TransformProxy->GetTransform();
@@ -238,6 +249,7 @@ void USpotLightGizmo::CreateOuterAngleGizmo()
 	OuterAngleGizmo->UpdateAngleFunc = [this](float NewAngle) { this->OnOuterAngleUpdate(NewAngle); };
 	OuterAngleGizmo->MaxAngle = 80.f;
 	OuterAngleGizmo->MinAngle = 1.f;
+	OuterAngleGizmo->TransactionDescription = LOCTEXT("SpotLightOuterAngle", "Spot Light Outer Angle");
 }
 
 void USpotLightGizmo::CreateInnerAngleGizmo()
@@ -253,6 +265,7 @@ void USpotLightGizmo::CreateInnerAngleGizmo()
 	InnerAngleGizmo->MaxAngle = 80.f;
 	InnerAngleGizmo->MinAngle = 1.f;
 	InnerAngleGizmo->ConeColor = FColor(150, 200, 255);
+	InnerAngleGizmo->TransactionDescription = LOCTEXT("SpotLightInnerAngle", "Spot Light Inner Angle");
 }
 
 void USpotLightGizmo::SetWorld(UWorld* InWorld)
@@ -288,6 +301,8 @@ void USpotLightGizmo::OnBeginDrag(const FInputDeviceRay& Ray)
 		bIsDragging = true;
 
 		UpdateHandleColors();
+
+		GetGizmoManager()->BeginUndoTransaction(LOCTEXT("SpotLightAttenuation", "Spot Light Attenuation"));
 	}
 }
 
@@ -313,6 +328,7 @@ void USpotLightGizmo::OnUpdateDrag(const FInputDeviceRay& Ray)
 
 	NewAttenuation = (NewAttenuation < 0) ? 0 : NewAttenuation;
 
+	LightActor->SpotLightComponent->Modify();
 	LightActor->SpotLightComponent->AttenuationRadius = NewAttenuation;
 	LightActor->SpotLightComponent->MarkRenderStateDirty();
 
@@ -332,6 +348,8 @@ void USpotLightGizmo::OnEndDrag(const FInputDeviceRay& Ray)
 	bIsDragging = false;
 
 	UpdateHandleColors();
+
+	GetGizmoManager()->EndUndoTransaction();
 }
 
 bool USpotLightGizmo::HitTest(const FRay& Ray, FHitResult& OutHit, FTransform& OutTransform)
@@ -433,4 +451,6 @@ void USpotLightGizmoInputBehavior::ForceEndCapture(const FInputCaptureData& data
 		Gizmo->OnEndDrag(LastWorldRay);
 		bInputDragCaptured = false;
 	}
-}
\ No newline at end of file
+}
+
+#undef LOCTEXT_NAMESPACE
\ No newline at end of file
diff --git a/Engine/Plugins/Experimental/GizmoEdMode/Source/LightGizmos/Public/ScalableConeGizmo.h b/Engine/Plugins/Experimental/GizmoEdMode/Source/LightGizmos/Public/ScalableConeGizmo.h
--- a/Engine/Plugins/Experimental/GizmoEdMode/Source/LightGizmos/Public/ScalableConeGizmo.h
+++ b/Engine/Plugins/Experimental/GizmoEdMode/Source/LightGizmos/Public/ScalableConeGizmo.h
@@ -77,6 +77,9 @@ public:
 	UPROPERTY()
 	float HitErrorThreshold{ 12.f };
 
+	// The text that will be used as the transaction description for undo/redo
+	UPROPERTY()
+	FText TransactionDescription;
 
 	/** Called when the Angle of the cone is changed. Sends new angle as parameter. */
 	TFunction<void(const float)> UpdateAngleFunc = nullptr;
diff --git a/Engine/Source/Runtime/Experimental/InteractiveToolsFramework/Private/BaseGizmos/ScalableSphereGizmo.cpp b/Engine/Source/Runtime/Experimental/InteractiveToolsFramework/Private/BaseGizmos/ScalableSphereGizmo.cpp
--- a/Engine/Source/Runtime/Experimental/InteractiveToolsFramework/Private/BaseGizmos/ScalableSphereGizmo.cpp
+++ b/Engine/Source/Runtime/Experimental/InteractiveToolsFramework/Private/BaseGizmos/ScalableSphereGizmo.cpp
@@ -11,6 +11,7 @@
 #include "BaseBehaviors/MouseHoverBehavior.h"
 #include "BaseGizmos/GizmoRenderingUtil.h"
 
+#define LOCTEXT_NAMESPACE "UScalableSphereGizmo"
 // UScalableSphereGizmoBuilder
 
 UInteractiveGizmo* UScalableSphereGizmoBuilder::BuildGizmo(const FToolBuilderState& SceneState) const
@@ -27,6 +28,8 @@ void UScalableSphereGizmo::Setup()
 
 	Radius = 1000.0f;
 
+	TransactionDescription = LOCTEXT("ScalableSphereGizmo", "Scale Sphere Gizmo");
+
 	UScalableSphereGizmoInputBehavior* ScalableSphereBehavior = NewObject<UScalableSphereGizmoInputBehavior>(this);
 	ScalableSphereBehavior->Initialize(this);
 	AddInputBehavior(ScalableSphereBehavior);
@@ -254,6 +257,8 @@ void UScalableSphereGizmo::OnBeginDrag(const FInputDeviceRay& Ray)
 		DragCurrentPositionProjected = DragStartWorldPosition;
 
 		bIsDragging = true;
+
+		GetGizmoManager()->BeginUndoTransaction(TransactionDescription);
 	}
 }
 
@@ -284,6 +289,7 @@ void UScalableSphereGizmo::OnUpdateDrag(const FInputDeviceRay& Ray)
 
 void UScalableSphereGizmo::OnEndDrag(const FInputDeviceRay& Ray)
 {
+	GetGizmoManager()->EndUndoTransaction();
 	bIsDragging = false;
 }
 
@@ -349,4 +355,6 @@ void UScalableSphereGizmoInputBehavior::ForceEndCapture(const FInputCaptureData&
 		bInputDragCaptured = false;
 		Gizmo->OnEndDrag(LastWorldRay);
 	}
-}
\ No newline at end of file
+}
+
+#undef LOCTEXT_NAMESPACE
\ No newline at end of file
diff --git a/Engine/Source/Runtime/Experimental/InteractiveToolsFramework/Public/BaseGizmos/ScalableSphereGizmo.h b/Engine/Source/Runtime/Experimental/InteractiveToolsFramework/Public/BaseGizmos/ScalableSphereGizmo.h
--- a/Engine/Source/Runtime/Experimental/InteractiveToolsFramework/Public/BaseGizmos/ScalableSphereGizmo.h
+++ b/Engine/Source/Runtime/Experimental/InteractiveToolsFramework/Public/BaseGizmos/ScalableSphereGizmo.h
@@ -62,6 +62,10 @@ public:
 	UPROPERTY()
 	float HitErrorThreshold{ 12.f };
 
+	// The text that will be used as the transaction description for undo/redo
+	UPROPERTY()
+	FText TransactionDescription;
+
 private:
 
 	// Check if the Circle represented by CircleNormal centered at the gizmos location has been hit by the Ray
