NewGizmos: Safeguard creation of new gizmos immediately after toggle

There was a chicken and egg problem with starting the gizmo edmode and creating the initial gizmos because the mode start hook was the only place where factories could register their gizmos with the GizmoManager, and the handlers could fire out of order.  This change simply delays the gizmo creation until the first Tick() after mode start so that all of the factories can register their gizmos prior to the creation.

#rnx

diff --git a/Engine/Plugins/Experimental/GizmoEdMode/Source/GizmoEdMode/Private/GizmoEdMode.cpp b/Engine/Plugins/Experimental/GizmoEdMode/Source/GizmoEdMode/Private/GizmoEdMode.cpp
--- a/Engine/Plugins/Experimental/GizmoEdMode/Source/GizmoEdMode/Private/GizmoEdMode.cpp
+++ b/Engine/Plugins/Experimental/GizmoEdMode/Source/GizmoEdMode/Private/GizmoEdMode.cpp
@@ -70,7 +70,7 @@ void UGizmoEdMode::DestroyGizmo()
 void UGizmoEdMode::Enter()
 {
 	Super::Enter();
-	RecreateGizmo();
+	bNeedInitialGizmos = true;
 	WidgetModeChangedHandle =
 	    GetModeManager()->OnWidgetModeChanged().AddLambda([this](UE::Widget::EWidgetMode) { RecreateGizmo(); });
 	GetModeManager()->SetShowWidget(false);
@@ -93,6 +93,11 @@ bool UGizmoEdMode::InputKey(FEditorViewportClient* ViewportClient, FViewport* Vi
 void UGizmoEdMode::Tick(FEditorViewportClient* ViewportClient, float DeltaTime)
 {
 	Super::Tick(ViewportClient, DeltaTime);
+	if ( bNeedInitialGizmos )
+	{
+		RecreateGizmo();
+		bNeedInitialGizmos = false;
+	}
 	if ( LastFactory )
 	{
 		LastFactory->ConfigureGridSnapping(GetDefault<ULevelEditorViewportSettings>()->GridEnabled,
diff --git a/Engine/Plugins/Experimental/GizmoEdMode/Source/GizmoEdMode/Public/GizmoEdMode.h b/Engine/Plugins/Experimental/GizmoEdMode/Source/GizmoEdMode/Public/GizmoEdMode.h
--- a/Engine/Plugins/Experimental/GizmoEdMode/Source/GizmoEdMode/Public/GizmoEdMode.h
+++ b/Engine/Plugins/Experimental/GizmoEdMode/Source/GizmoEdMode/Public/GizmoEdMode.h
@@ -49,4 +49,6 @@ private:
 	TArray<UInteractiveGizmo*> InteractiveGizmos;
 
 	FDelegateHandle WidgetModeChangedHandle;
+
+	bool bNeedInitialGizmos{false};
 };
