ToolsFramework: fix up calls to implicit FInputDeviceRay constructor in experimental plugins
#rnx

diff --git a/Engine/Plugins/Experimental/GizmoEdMode/Source/LightGizmos/Private/DirectionalLightGizmo.cpp b/Engine/Plugins/Experimental/GizmoEdMode/Source/LightGizmos/Private/DirectionalLightGizmo.cpp
--- a/Engine/Plugins/Experimental/GizmoEdMode/Source/LightGizmos/Private/DirectionalLightGizmo.cpp
+++ b/Engine/Plugins/Experimental/GizmoEdMode/Source/LightGizmos/Private/DirectionalLightGizmo.cpp
@@ -462,11 +462,11 @@ FInputCaptureUpdate UDirectionalLightGizmoInputBehavior::UpdateCapture(const FIn
 	if (IsReleased(input))
 	{
 		bInputDragCaptured = false;
-		Gizmo->OnEndDrag(LastWorldRay);
+		Gizmo->OnEndDrag(FInputDeviceRay(LastWorldRay));
 		return FInputCaptureUpdate::End();
 	}
 
-	Gizmo->OnUpdateDrag(LastWorldRay);
+	Gizmo->OnUpdateDrag(FInputDeviceRay(LastWorldRay));
 
 	return FInputCaptureUpdate::Continue();
 }
@@ -475,7 +475,7 @@ void UDirectionalLightGizmoInputBehavior::ForceEndCapture(const FInputCaptureDat
 {
 	if (bInputDragCaptured)
 	{
-		Gizmo->OnEndDrag(LastWorldRay);
+		Gizmo->OnEndDrag(FInputDeviceRay(LastWorldRay));
 		bInputDragCaptured = false;
 	}
 }
diff --git a/Engine/Plugins/Experimental/GizmoEdMode/Source/LightGizmos/Private/ScalableConeGizmo.cpp b/Engine/Plugins/Experimental/GizmoEdMode/Source/LightGizmos/Private/ScalableConeGizmo.cpp
--- a/Engine/Plugins/Experimental/GizmoEdMode/Source/LightGizmos/Private/ScalableConeGizmo.cpp
+++ b/Engine/Plugins/Experimental/GizmoEdMode/Source/LightGizmos/Private/ScalableConeGizmo.cpp
@@ -357,11 +357,11 @@ FInputCaptureUpdate UScalableConeGizmoInputBehavior::UpdateCapture(const FInputD
 	if (IsReleased(input))
 	{
 		bInputDragCaptured = false;
-		Gizmo->OnEndDrag(LastWorldRay);
+		Gizmo->OnEndDrag(FInputDeviceRay(LastWorldRay));
 		return FInputCaptureUpdate::End();
 	}
 
-	Gizmo->OnUpdateDrag(LastWorldRay);
+	Gizmo->OnUpdateDrag(FInputDeviceRay(LastWorldRay));
 
 	return FInputCaptureUpdate::Continue();
 }
@@ -371,7 +371,7 @@ void UScalableConeGizmoInputBehavior::ForceEndCapture(const FInputCaptureData& d
 	if (bInputDragCaptured)
 	{
 		bInputDragCaptured = false;
-		Gizmo->OnEndDrag(LastWorldRay);
+		Gizmo->OnEndDrag(FInputDeviceRay(LastWorldRay));
 	}
 }
 
diff --git a/Engine/Plugins/Experimental/GizmoEdMode/Source/LightGizmos/Private/SpotLightGizmo.cpp b/Engine/Plugins/Experimental/GizmoEdMode/Source/LightGizmos/Private/SpotLightGizmo.cpp
--- a/Engine/Plugins/Experimental/GizmoEdMode/Source/LightGizmos/Private/SpotLightGizmo.cpp
+++ b/Engine/Plugins/Experimental/GizmoEdMode/Source/LightGizmos/Private/SpotLightGizmo.cpp
@@ -435,11 +435,11 @@ FInputCaptureUpdate USpotLightGizmoInputBehavior::UpdateCapture(const FInputDevi
 	if (IsReleased(input))
 	{
 		bInputDragCaptured = false;
-		Gizmo->OnEndDrag(LastWorldRay);
+		Gizmo->OnEndDrag(FInputDeviceRay(LastWorldRay));
 		return FInputCaptureUpdate::End();
 	}
 
-	Gizmo->OnUpdateDrag(LastWorldRay);
+	Gizmo->OnUpdateDrag(FInputDeviceRay(LastWorldRay));
 
 	return FInputCaptureUpdate::Continue();
 }
@@ -448,7 +448,7 @@ void USpotLightGizmoInputBehavior::ForceEndCapture(const FInputCaptureData& data
 {
 	if (bInputDragCaptured)
 	{
-		Gizmo->OnEndDrag(LastWorldRay);
+		Gizmo->OnEndDrag(FInputDeviceRay(LastWorldRay));
 		bInputDragCaptured = false;
 	}
 }
