Making the new gizmo module on by default but you have to use Editor.ShowGizmoOptions to show the menu option now. Fixing the gizmo mode using a toolkit and fixing a condition where toggling the mode on and off would crash due to duplicate registration.

diff --git a/Engine/Plugins/Experimental/GizmoEdMode/GizmoEdMode.uplugin b/Engine/Plugins/Experimental/GizmoEdMode/GizmoEdMode.uplugin
--- a/Engine/Plugins/Experimental/GizmoEdMode/GizmoEdMode.uplugin
+++ b/Engine/Plugins/Experimental/GizmoEdMode/GizmoEdMode.uplugin
@@ -13,7 +13,7 @@
 	"CanContainContent": false,
 	"IsBetaVersion": false,
 	"IsExperimentalVersion": true,
-	"Installed": false,
+	"Installed": true,
 	"Modules": [
 		{
 			"Name": "GizmoEdMode",
diff --git a/Engine/Plugins/Experimental/GizmoEdMode/Source/GizmoEdMode/Public/GizmoEdMode.h b/Engine/Plugins/Experimental/GizmoEdMode/Source/GizmoEdMode/Public/GizmoEdMode.h
--- a/Engine/Plugins/Experimental/GizmoEdMode/Source/GizmoEdMode/Public/GizmoEdMode.h
+++ b/Engine/Plugins/Experimental/GizmoEdMode/Source/GizmoEdMode/Public/GizmoEdMode.h
@@ -28,7 +28,10 @@ public:
 	UGizmoEdMode();
 
 	void AddFactory(TScriptInterface<IAssetEditorGizmoFactory> GizmoFactory);
-
+	virtual bool UsesToolkits() const override
+	{
+		return false;
+	}
 private:
 	void ActorSelectionChangeNotify() override;
 	void Enter() override;
diff --git a/Engine/Plugins/Experimental/GizmoEdMode/Source/LightGizmos/LightGizmos.Build.cs b/Engine/Plugins/Experimental/GizmoEdMode/Source/LightGizmos/LightGizmos.Build.cs
--- a/Engine/Plugins/Experimental/GizmoEdMode/Source/LightGizmos/LightGizmos.Build.cs
+++ b/Engine/Plugins/Experimental/GizmoEdMode/Source/LightGizmos/LightGizmos.Build.cs
@@ -41,7 +41,8 @@ public class LightGizmos : ModuleRules
 				"CoreUObject",
 				"EditorInteractiveToolsFramework",
 				"InteractiveToolsFramework",
-				"GizmoEdMode"
+				"GizmoEdMode",
+				"LevelEditor",
 				// ... add private dependencies that you statically link with here ...
 			}
 			);
diff --git a/Engine/Plugins/Experimental/GizmoEdMode/Source/LightGizmos/Private/LightGizmosModule.cpp b/Engine/Plugins/Experimental/GizmoEdMode/Source/LightGizmos/Private/LightGizmosModule.cpp
--- a/Engine/Plugins/Experimental/GizmoEdMode/Source/LightGizmos/Private/LightGizmosModule.cpp
+++ b/Engine/Plugins/Experimental/GizmoEdMode/Source/LightGizmos/Private/LightGizmosModule.cpp
@@ -13,6 +13,9 @@
 #include "ScalableConeGizmo.h"
 #include "SpotLightGizmo.h"
 #include "DirectionalLightGizmo.h"
+#include "LevelEditor.h"
+#include "EdModeInteractiveToolsContext.h"
+#include "InteractiveGizmoManager.h"
 
 #define LOCTEXT_NAMESPACE "FLightGizmosModule"
 
@@ -21,34 +24,33 @@ FString FLightGizmosModule::SpotLightGizmoType = TEXT("SpotLightGizmoType");
 FString FLightGizmosModule::ScalableConeGizmoType = TEXT("ScalableConeGizmoType");
 FString FLightGizmosModule::DirectionalLightGizmoType = TEXT("DirectionalLightGizmoType");
 
-void FLightGizmosModule::OnPostEngineInit()
+void FLightGizmosModule::OnLevelEditorCreated(TSharedPtr<ILevelEditor> InLevelEditor)
 {
-	GLevelEditorModeTools().OnEditorModeIDChanged().AddLambda([](const FEditorModeID& ModeID, bool IsEnteringMode) 
+	TWeakPtr<ILevelEditor> LevelEditorPtr = InLevelEditor;
+	InLevelEditor->GetEditorModeManager().OnEditorModeIDChanged().AddLambda([LevelEditorPtr](const FEditorModeID& ModeID, bool IsEnteringMode)
 	{
-		if ((ModeID == GetDefault<UGizmoEdMode>()->GetID()))
+		if (ModeID == GetDefault<UGizmoEdMode>()->GetID() && LevelEditorPtr.IsValid())
 		{
-			UEdMode* EdMode = GLevelEditorModeTools().GetActiveScriptableMode(GetDefault<UGizmoEdMode>()->GetID());
+			UEdMode* EdMode = LevelEditorPtr.Pin()->GetEditorModeManager().GetActiveScriptableMode(GetDefault<UGizmoEdMode>()->GetID());
 			UGizmoEdMode* GizmoEdMode = Cast<UGizmoEdMode>(EdMode);
 
 			 // Register the factories and gizmos if we are entering the new Gizmo Mode
-			if (IsEnteringMode)
+			if (IsEnteringMode && GizmoEdMode)
 			{
-				GizmoEdMode->GetToolManager()->GetPairedGizmoManager()->RegisterGizmoType(PointLightGizmoType, NewObject<UPointLightGizmoBuilder>());
-				GizmoEdMode->GetToolManager()->GetPairedGizmoManager()->RegisterGizmoType(ScalableConeGizmoType, NewObject<UScalableConeGizmoBuilder>());
-				GizmoEdMode->GetToolManager()->GetPairedGizmoManager()->RegisterGizmoType(SpotLightGizmoType, NewObject<USpotLightGizmoBuilder>());
-				GizmoEdMode->GetToolManager()->GetPairedGizmoManager()->RegisterGizmoType(DirectionalLightGizmoType, NewObject<UDirectionalLightGizmoBuilder>());
+				LevelEditorPtr.Pin()->GetEditorModeManager().GetInteractiveToolsContext()->GizmoManager->RegisterGizmoType(PointLightGizmoType, NewObject<UPointLightGizmoBuilder>());
+				LevelEditorPtr.Pin()->GetEditorModeManager().GetInteractiveToolsContext()->GizmoManager->RegisterGizmoType(ScalableConeGizmoType, NewObject<UScalableConeGizmoBuilder>());
+				LevelEditorPtr.Pin()->GetEditorModeManager().GetInteractiveToolsContext()->GizmoManager->RegisterGizmoType(SpotLightGizmoType, NewObject<USpotLightGizmoBuilder>());
+				LevelEditorPtr.Pin()->GetEditorModeManager().GetInteractiveToolsContext()->GizmoManager->RegisterGizmoType(DirectionalLightGizmoType, NewObject<UDirectionalLightGizmoBuilder>());
 				GizmoEdMode->AddFactory(NewObject<UPointLightGizmoFactory>());
 				GizmoEdMode->AddFactory(NewObject<UDirectionalLightGizmoFactory>());
 				GizmoEdMode->AddFactory(NewObject<USpotLightGizmoFactory>());
 			}
 			else
 			{
-				//TODO: Deregistering on exit currently causes a crash
-
-				//GizmoEdMode->GetToolManager()->GetPairedGizmoManager()->DeregisterGizmoType(PointLightGizmoType);
-				//GizmoEdMode->GetToolManager()->GetPairedGizmoManager()->DeregisterGizmoType(ScalableConeGizmoType);
-				//GizmoEdMode->GetToolManager()->GetPairedGizmoManager()->DeregisterGizmoType(SpotLightGizmoType);
-				//GizmoEdMode->GetToolManager()->GetPairedGizmoManager()->DeregisterGizmoType(DirectionalLightGizmoType);
+				LevelEditorPtr.Pin()->GetEditorModeManager().GetInteractiveToolsContext()->GizmoManager->DeregisterGizmoType(PointLightGizmoType);
+				LevelEditorPtr.Pin()->GetEditorModeManager().GetInteractiveToolsContext()->GizmoManager->DeregisterGizmoType(ScalableConeGizmoType);
+				LevelEditorPtr.Pin()->GetEditorModeManager().GetInteractiveToolsContext()->GizmoManager->DeregisterGizmoType(SpotLightGizmoType);
+				LevelEditorPtr.Pin()->GetEditorModeManager().GetInteractiveToolsContext()->GizmoManager->DeregisterGizmoType(DirectionalLightGizmoType);
 			}
 			
 		}
@@ -58,7 +60,8 @@ void FLightGizmosModule::OnPostEngineInit()
 
 void FLightGizmosModule::StartupModule()
 {
-	FCoreDelegates::OnPostEngineInit.AddRaw(this, &FLightGizmosModule::OnPostEngineInit);
+	FLevelEditorModule& LevelEditor = FModuleManager::Get().LoadModuleChecked<FLevelEditorModule>("LevelEditor");
+	LevelEditor.OnLevelEditorCreated().AddRaw(this, &FLightGizmosModule::OnLevelEditorCreated);
 }
 
 void FLightGizmosModule::ShutdownModule()
diff --git a/Engine/Plugins/Experimental/GizmoEdMode/Source/LightGizmos/Public/LightGizmosModule.h b/Engine/Plugins/Experimental/GizmoEdMode/Source/LightGizmos/Public/LightGizmosModule.h
--- a/Engine/Plugins/Experimental/GizmoEdMode/Source/LightGizmos/Public/LightGizmosModule.h
+++ b/Engine/Plugins/Experimental/GizmoEdMode/Source/LightGizmos/Public/LightGizmosModule.h
@@ -5,12 +5,14 @@
 #include "CoreMinimal.h"
 #include "Modules/ModuleManager.h"
 
+class ILevelEditor;
+
 class FLightGizmosModule : public IModuleInterface
 {
 public:
 
 	/** IModuleInterface implementation */
-	virtual void OnPostEngineInit();
+	virtual void OnLevelEditorCreated(TSharedPtr<ILevelEditor> InLevelEditor);
 
 	virtual void StartupModule() override;
 	virtual void ShutdownModule() override;
diff --git a/Engine/Plugins/Experimental/GizmoFramework/Source/NewGizmoEnable/Private/NewGizmoEnable.cpp b/Engine/Plugins/Experimental/GizmoFramework/Source/NewGizmoEnable/Private/NewGizmoEnable.cpp
--- a/Engine/Plugins/Experimental/GizmoFramework/Source/NewGizmoEnable/Private/NewGizmoEnable.cpp
+++ b/Engine/Plugins/Experimental/GizmoFramework/Source/NewGizmoEnable/Private/NewGizmoEnable.cpp
@@ -18,7 +18,8 @@
 #include "GizmoEdMode.h"
 
 static const FName NewGizmoEnableTabName("NewGizmoEnable");
-
+static FAutoConsoleCommand ShowGizmoOptions(TEXT("Editor.ShowGizmoOptions"), TEXT("Shows the new gizmo mode in menus"), FConsoleCommandDelegate::CreateStatic(&FNewGizmoEnableModule::ShowGizmoOptions));
+bool FNewGizmoEnableModule::bTestingModeEnabled = false;
 #define LOCTEXT_NAMESPACE "FNewGizmoEnableModule"
 
 void FNewGizmoEnableModule::StartupModule()
@@ -30,7 +31,8 @@ void FNewGizmoEnableModule::StartupModule()
 		FNewGizmoEnableCommands::Get().ToggleNewGizmos,
 		FExecuteAction::CreateRaw(this, &FNewGizmoEnableModule::ToggleNewGizmosActive),
 		FCanExecuteAction(),
-		FGetActionCheckState::CreateRaw(this, &FNewGizmoEnableModule::AreNewGizmosActive));
+		FGetActionCheckState::CreateRaw(this, &FNewGizmoEnableModule::AreNewGizmosActive),
+		FIsActionButtonVisible::CreateRaw(this, &FNewGizmoEnableModule::CanToggleNewGizmos));
 
 	UToolMenus::RegisterStartupCallback(FSimpleMulticastDelegate::FDelegate::CreateRaw(this, &FNewGizmoEnableModule::RegisterMenus));
 }
@@ -66,6 +68,16 @@ ECheckBoxState FNewGizmoEnableModule::AreNewGizmosActive()
 	                                                                                   ECheckBoxState::Unchecked;
 }
 
+bool FNewGizmoEnableModule::CanToggleNewGizmos()
+{
+	return FNewGizmoEnableModule::bTestingModeEnabled;
+}
+
+void FNewGizmoEnableModule::ShowGizmoOptions()
+{
+	FNewGizmoEnableModule::bTestingModeEnabled = !FNewGizmoEnableModule::bTestingModeEnabled;
+}
+
 void FNewGizmoEnableModule::RegisterMenus()
 {
 	// Owner will be used for cleanup in call to UToolMenus::UnregisterOwner
diff --git a/Engine/Plugins/Experimental/GizmoFramework/Source/NewGizmoEnable/Public/NewGizmoEnable.h b/Engine/Plugins/Experimental/GizmoFramework/Source/NewGizmoEnable/Public/NewGizmoEnable.h
--- a/Engine/Plugins/Experimental/GizmoFramework/Source/NewGizmoEnable/Public/NewGizmoEnable.h
+++ b/Engine/Plugins/Experimental/GizmoFramework/Source/NewGizmoEnable/Public/NewGizmoEnable.h
@@ -20,6 +20,8 @@ public:
 	/** Action Callbacks */
 	void ToggleNewGizmosActive();
 	ECheckBoxState AreNewGizmosActive();
+	bool CanToggleNewGizmos();
+	static void ShowGizmoOptions();
 private:
 
 	void RegisterMenus();
@@ -29,4 +31,6 @@ private:
 private:
 	TSharedPtr<class FUICommandList> PluginCommands;
 	bool bGizmosEnabled = false;
+	/** Whether currently in testing mode or not. */
+	static bool bTestingModeEnabled;
 };
