Fix Gizmo Editor Mode plugin to appear in Built-In->Experimental in plugin manager, instead of in Installed.

diff --git a/Engine/Plugins/Experimental/GizmoEdMode/GizmoEdMode.uplugin b/Engine/Plugins/Experimental/GizmoEdMode/GizmoEdMode.uplugin
--- a/Engine/Plugins/Experimental/GizmoEdMode/GizmoEdMode.uplugin
+++ b/Engine/Plugins/Experimental/GizmoEdMode/GizmoEdMode.uplugin
@@ -4,7 +4,7 @@
 	"VersionName": "1.0",
 	"FriendlyName": "Gizmo Editor Mode",
 	"Description": "Editor mode to manage InteractiveToolFramework based global TRS gizmos",
-	"Category": "Assets",
+	"Category": "Experimental",
 	"CreatedBy": "Epic Games, Inc.",
 	"CreatedByURL": "http://epicgames.com",
 	"DocsURL": "",
@@ -13,7 +13,7 @@
 	"CanContainContent": false,
 	"IsBetaVersion": false,
 	"IsExperimentalVersion": true,
-	"Installed": true,
+	"Installed": false,
 	"Modules": [
 		{
 			"Name": "GizmoEdMode",
