Fixed no-unity build: Added missing include in ScalableConeGizmo.cpp

diff --git a/Engine/Plugins/Experimental/GizmoEdMode/Source/LightGizmos/Private/ScalableConeGizmo.cpp b/Engine/Plugins/Experimental/GizmoEdMode/Source/LightGizmos/Private/ScalableConeGizmo.cpp
--- a/Engine/Plugins/Experimental/GizmoEdMode/Source/LightGizmos/Private/ScalableConeGizmo.cpp
+++ b/Engine/Plugins/Experimental/GizmoEdMode/Source/LightGizmos/Private/ScalableConeGizmo.cpp
@@ -7,6 +7,7 @@
 #include "BaseGizmos/GizmoMath.h"
 #include "BaseGizmos/GizmoRenderingUtil.h"
 #include "BaseBehaviors/MouseHoverBehavior.h"
+#include "SceneManagement.h"
 
 #define LOCTEXT_NAMESPACE "UScalableSphereGizmo"
 
