[Truncation Warnings] Fix warnings in GizmoEdMode plugin modules

#rnx

diff --git a/Engine/Plugins/Experimental/GizmoEdMode/Source/LightGizmos/Private/DirectionalLightGizmo.cpp b/Engine/Plugins/Experimental/GizmoEdMode/Source/LightGizmos/Private/DirectionalLightGizmo.cpp
--- a/Engine/Plugins/Experimental/GizmoEdMode/Source/LightGizmos/Private/DirectionalLightGizmo.cpp
+++ b/Engine/Plugins/Experimental/GizmoEdMode/Source/LightGizmos/Private/DirectionalLightGizmo.cpp
@@ -120,7 +120,7 @@ void UDirectionalLightGizmo::Render(IToolsContextRenderAPI* RenderAPI)
 	RenderAPI->GetPrimitiveDrawInterface()->DrawLine(Start, LineEndWorld, FLinearColor::Red, SDPG_Foreground);
 
 	// Figure out the angle of the arc between the line and the arrow
-	float ArcAngle = FMath::Acos(FVector::DotProduct(LineDir, ArrowDir));
+	float ArcAngle = FMath::Acos(static_cast<float>(FVector::DotProduct(LineDir, ArrowDir)));
 	ArcAngle = FMath::RadiansToDegrees(ArcAngle);
 
 	// Draw an arc between the arrow and the line
diff --git a/Engine/Plugins/Experimental/GizmoEdMode/Source/LightGizmos/Private/ScalableConeGizmo.cpp b/Engine/Plugins/Experimental/GizmoEdMode/Source/LightGizmos/Private/ScalableConeGizmo.cpp
--- a/Engine/Plugins/Experimental/GizmoEdMode/Source/LightGizmos/Private/ScalableConeGizmo.cpp
+++ b/Engine/Plugins/Experimental/GizmoEdMode/Source/LightGizmos/Private/ScalableConeGizmo.cpp
@@ -289,7 +289,7 @@ void UScalableConeGizmo::OnUpdateDrag(const FInputDeviceRay& Ray)
 	FVector End = AxisNearestPt - GizmoLocation;
 	End.Normalize();
 
-	float DotP = FVector::DotProduct(Start, End);
+	float DotP = static_cast<float>(FVector::DotProduct(Start, End));
 
 	float DeltaAngle = FMath::Acos(DotP);
 
@@ -297,8 +297,8 @@ void UScalableConeGizmo::OnUpdateDrag(const FInputDeviceRay& Ray)
 
 	FVector Forward = ActiveTarget->GetTransform().GetRotation().Vector();
 
-	float StartAngle = FMath::Acos(FVector::DotProduct(Start, Forward));
-	float EndAngle = FMath::Acos(FVector::DotProduct(End, Forward));
+	float StartAngle = FMath::Acos(static_cast<float>(FVector::DotProduct(Start, Forward)));
+	float EndAngle = FMath::Acos(static_cast<float>(FVector::DotProduct(End, Forward)));
 
 	if (StartAngle > EndAngle)
 	{
