[Truncation Warnings] Enable truncation warnings in build.cs files for fixed modules

#rnx

diff --git a/Engine/Plugins/Experimental/GizmoEdMode/Source/GizmoEdMode/GizmoEdMode.Build.cs b/Engine/Plugins/Experimental/GizmoEdMode/Source/GizmoEdMode/GizmoEdMode.Build.cs
--- a/Engine/Plugins/Experimental/GizmoEdMode/Source/GizmoEdMode/GizmoEdMode.Build.cs
+++ b/Engine/Plugins/Experimental/GizmoEdMode/Source/GizmoEdMode/GizmoEdMode.Build.cs
@@ -6,6 +6,9 @@ public class GizmoEdMode : ModuleRules
 {
 	public GizmoEdMode(ReadOnlyTargetRules Target) : base(Target)
 	{
+		// Enable truncation warnings in this module
+		CppCompileWarningSettings.UnsafeTypeCastWarningLevel = WarningLevel.Warning;
+
 		PCHUsage = ModuleRules.PCHUsageMode.UseExplicitOrSharedPCHs;
 
 		PublicIncludePaths.AddRange(
diff --git a/Engine/Plugins/Experimental/GizmoEdMode/Source/LightGizmos/LightGizmos.Build.cs b/Engine/Plugins/Experimental/GizmoEdMode/Source/LightGizmos/LightGizmos.Build.cs
--- a/Engine/Plugins/Experimental/GizmoEdMode/Source/LightGizmos/LightGizmos.Build.cs
+++ b/Engine/Plugins/Experimental/GizmoEdMode/Source/LightGizmos/LightGizmos.Build.cs
@@ -6,6 +6,9 @@ public class LightGizmos : ModuleRules
 {
 	public LightGizmos(ReadOnlyTargetRules Target) : base(Target)
 	{
+		// Enable truncation warnings in this module
+		CppCompileWarningSettings.UnsafeTypeCastWarningLevel = WarningLevel.Warning;
+
 		PCHUsage = ModuleRules.PCHUsageMode.UseExplicitOrSharedPCHs;
 
 		PublicIncludePaths.AddRange(
diff --git a/Engine/Plugins/Experimental/GizmoFramework/Source/GizmoSettings/GizmoSettings.Build.cs b/Engine/Plugins/Experimental/GizmoFramework/Source/GizmoSettings/GizmoSettings.Build.cs
--- a/Engine/Plugins/Experimental/GizmoFramework/Source/GizmoSettings/GizmoSettings.Build.cs
+++ b/Engine/Plugins/Experimental/GizmoFramework/Source/GizmoSettings/GizmoSettings.Build.cs
@@ -6,6 +6,9 @@ namespace UnrealBuildTool.Rules
 	{
 		public GizmoSettings(ReadOnlyTargetRules Target) : base(Target)
 		{
+			// Enable truncation warnings in this module
+			CppCompileWarningSettings.UnsafeTypeCastWarningLevel = WarningLevel.Warning;
+
 			PrivateDependencyModuleNames.AddRange(
 				new string[]
 				{
diff --git a/Engine/Source/Editor/EditorFramework/EditorFramework.Build.cs b/Engine/Source/Editor/EditorFramework/EditorFramework.Build.cs
--- a/Engine/Source/Editor/EditorFramework/EditorFramework.Build.cs
+++ b/Engine/Source/Editor/EditorFramework/EditorFramework.Build.cs
@@ -6,6 +6,9 @@ public class EditorFramework : ModuleRules
 {
 	public EditorFramework(ReadOnlyTargetRules Target) : base(Target)
 	{
+		// Enable truncation warnings in this module
+		CppCompileWarningSettings.UnsafeTypeCastWarningLevel = WarningLevel.Warning;
+
 		PrivateDependencyModuleNames.AddRange(
 			new string[] {
 				"Core",
diff --git a/Engine/Source/Editor/GameProjectGeneration/GameProjectGeneration.Build.cs b/Engine/Source/Editor/GameProjectGeneration/GameProjectGeneration.Build.cs
--- a/Engine/Source/Editor/GameProjectGeneration/GameProjectGeneration.Build.cs
+++ b/Engine/Source/Editor/GameProjectGeneration/GameProjectGeneration.Build.cs
@@ -6,6 +6,9 @@ public class GameProjectGeneration : ModuleRules
 {
 	public GameProjectGeneration(ReadOnlyTargetRules Target) : base(Target)
 	{
+		// Enable truncation warnings in this module
+		CppCompileWarningSettings.UnsafeTypeCastWarningLevel = WarningLevel.Warning;
+
 		PublicDependencyModuleNames.AddRange(
 			new string[] {
 				"HardwareTargeting",
diff --git a/Engine/Source/Editor/UndoHistoryEditor/UndoHistoryEditor.Build.cs b/Engine/Source/Editor/UndoHistoryEditor/UndoHistoryEditor.Build.cs
--- a/Engine/Source/Editor/UndoHistoryEditor/UndoHistoryEditor.Build.cs
+++ b/Engine/Source/Editor/UndoHistoryEditor/UndoHistoryEditor.Build.cs
@@ -6,6 +6,9 @@ namespace UnrealBuildTool.Rules
 	{
 		public UndoHistoryEditor(ReadOnlyTargetRules Target) : base(Target)
 		{
+			// Enable truncation warnings in this module
+			CppCompileWarningSettings.UnsafeTypeCastWarningLevel = WarningLevel.Warning;
+
 			PublicDependencyModuleNames.AddRange(
 				new string[] {
 					"Core"
